Copying //UE4/Dev-Framework to //UE4/Dev-Main (Source: //UE4/Dev-Framework @ 3967517)

diff --git a/Engine/Source/Runtime/UMG/Public/Blueprint/WidgetBlueprintLibrary.h b/Engine/Source/Runtime/UMG/Public/Blueprint/WidgetBlueprintLibrary.h
--- a/Engine/Source/Runtime/UMG/Public/Blueprint/WidgetBlueprintLibrary.h
+++ b/Engine/Source/Runtime/UMG/Public/Blueprint/WidgetBlueprintLibrary.h
@@ -231,19 +231,19 @@ public:
 	 * Gets the resource object on a brush.  This could be a UTexture2D or a UMaterialInterface.
 	 */
 	UFUNCTION(BlueprintPure, Category="Widget|Brush")
-	static UObject* GetBrushResource(UPARAM(ref) FSlateBrush& Brush);
+	static UObject* GetBrushResource(const FSlateBrush& Brush);
 
 	/**
 	 * Gets the brush resource as a texture 2D.
 	 */
 	UFUNCTION(BlueprintPure, Category="Widget|Brush")
-	static UTexture2D* GetBrushResourceAsTexture2D(UPARAM(ref) FSlateBrush& Brush);
+	static UTexture2D* GetBrushResourceAsTexture2D(const FSlateBrush& Brush);
 
 	/**
 	 * Gets the brush resource as a material.
 	 */
 	UFUNCTION(BlueprintPure, Category="Widget|Brush")
-	static UMaterialInterface* GetBrushResourceAsMaterial(UPARAM(ref) FSlateBrush& Brush);
+	static UMaterialInterface* GetBrushResourceAsMaterial(const FSlateBrush& Brush);
 
 	/**
 	 * Sets the resource on a brush to be a UTexture2D.
diff --git a/Templates/TP_TopDown/Source/TP_TopDown/TP_TopDown.Build.cs b/Templates/TP_TopDown/Source/TP_TopDown/TP_TopDown.Build.cs
--- a/Templates/TP_TopDown/Source/TP_TopDown/TP_TopDown.Build.cs
+++ b/Templates/TP_TopDown/Source/TP_TopDown/TP_TopDown.Build.cs
@@ -8,6 +8,6 @@ public class TP_TopDown : ModuleRules
 	{
 		PCHUsage = PCHUsageMode.UseExplicitOrSharedPCHs;
 
-		PublicDependencyModuleNames.AddRange(new string[] { "Core", "CoreUObject", "Engine", "InputCore", "HeadMountedDisplay" });
-	}
+        PublicDependencyModuleNames.AddRange(new string[] { "Core", "CoreUObject", "Engine", "InputCore", "HeadMountedDisplay", "NavigationSystem", "AIModule" });
+    }
 }
diff --git a/Templates/TP_TopDown/Source/TP_TopDown/TP_TopDownPlayerController.cpp b/Templates/TP_TopDown/Source/TP_TopDown/TP_TopDownPlayerController.cpp
--- a/Templates/TP_TopDown/Source/TP_TopDown/TP_TopDownPlayerController.cpp
+++ b/Templates/TP_TopDown/Source/TP_TopDown/TP_TopDownPlayerController.cpp
@@ -1,7 +1,7 @@
 // Copyright 1998-2018 Epic Games, Inc. All Rights Reserved.
 
 #include "TP_TopDownPlayerController.h"
-#include "AI/Navigation/NavigationSystem.h"
+#include "Blueprint/AIBlueprintHelperLibrary.h"
 #include "Runtime/Engine/Classes/Components/DecalComponent.h"
 #include "HeadMountedDisplayFunctionLibrary.h"
 #include "TP_TopDownCharacter.h"
@@ -52,7 +52,7 @@ void ATP_TopDownPlayerController::MoveToMouseCursor()
 		{
 			if (MyPawn->GetCursorToWorld())
 			{
-				UNavigationSystem::SimpleMoveToLocation(this, MyPawn->GetCursorToWorld()->GetComponentLocation());
+				UAIBlueprintHelperLibrary::SimpleMoveToLocation(this, MyPawn->GetCursorToWorld()->GetComponentLocation());
 			}
 		}
 	}
@@ -89,13 +89,12 @@ void ATP_TopDownPlayerController::SetNewMoveDestination(const FVector DestLocati
 	APawn* const MyPawn = GetPawn();
 	if (MyPawn)
 	{
-		UNavigationSystem* const NavSys = GetWorld()->GetNavigationSystem();
 		float const Distance = FVector::Dist(DestLocation, MyPawn->GetActorLocation());
 
 		// We need to issue move command only if far enough in order for walk animation to play correctly
-		if (NavSys && (Distance > 120.0f))
+		if ((Distance > 120.0f))
 		{
-			NavSys->SimpleMoveToLocation(this, DestLocation);
+			UAIBlueprintHelperLibrary::SimpleMoveToLocation(this, DestLocation);
 		}
 	}
 }
diff --git a/Templates/TP_TopDownBP/Config/DefaultEngine.ini b/Templates/TP_TopDownBP/Config/DefaultEngine.ini
--- a/Templates/TP_TopDownBP/Config/DefaultEngine.ini
+++ b/Templates/TP_TopDownBP/Config/DefaultEngine.ini
@@ -11,7 +11,7 @@ ThreePlayerSplitscreenLayout=FavorTop
 GlobalDefaultGameMode=/Game/TopDownBP/Blueprints/TopDownGameMode.TopDownGameMode_C
 GlobalDefaultServerGameMode=None
 
-[/Script/Engine.RecastNavMesh]
+[/Script/NavigationSystem.RecastNavMesh]
 bDrawPolyEdges=False
 bDistinctlyDrawTilesBeingBuilt=True
 DrawOffset=10.000000
