3rd Person and Top Down Template overhaul: First iteration of C++ (comments missing, some more cleanup outstanding), cursor for top down moved to shared folder, map updates, all blue boxes simulate physics (needs testing), BP  missing comments, level prototyping materials cleaned up.

diff --git a/Templates/TP_ThirdPerson/Config/DefaultEngine.ini b/Templates/TP_ThirdPerson/Config/DefaultEngine.ini
--- a/Templates/TP_ThirdPerson/Config/DefaultEngine.ini
+++ b/Templates/TP_ThirdPerson/Config/DefaultEngine.ini
@@ -5,3 +5,7 @@ GlobalDefaultGameMode="/Script/TP_ThirdPerson.TP_ThirdPersonGameMode"
 
 [/Script/IOSRuntimeSettings.IOSRuntimeSettings]
 MinimumiOSVersion=IOS_14
+
+[/Script/Engine.RendererSettings]
+r.Shadow.Virtual.Enable=1
+
diff --git a/Templates/TP_ThirdPerson/Source/TP_ThirdPerson/TP_ThirdPersonCharacter.cpp b/Templates/TP_ThirdPerson/Source/TP_ThirdPerson/TP_ThirdPersonCharacter.cpp
--- a/Templates/TP_ThirdPerson/Source/TP_ThirdPerson/TP_ThirdPersonCharacter.cpp
+++ b/Templates/TP_ThirdPerson/Source/TP_ThirdPerson/TP_ThirdPersonCharacter.cpp
@@ -16,9 +16,8 @@ ATP_ThirdPersonCharacter::ATP_ThirdPersonCharacter()
 	// Set size for collision capsule
 	GetCapsuleComponent()->InitCapsuleSize(42.f, 96.0f);
 
-	// set our turn rates for input
-	BaseTurnRate = 50.f;
-	BaseLookUpRate = 50.f;
+	// set our turn rate for input
+	TurnRateGamepad = 50.f;
 
 	// Don't rotate when the controller rotates. Let that just affect the camera.
 	bUseControllerRotationPitch = false;
@@ -27,9 +26,15 @@ ATP_ThirdPersonCharacter::ATP_ThirdPersonCharacter()
 
 	// Configure character movement
 	GetCharacterMovement()->bOrientRotationToMovement = true; // Character moves in the direction of input...	
-	GetCharacterMovement()->RotationRate = FRotator(0.0f, 540.0f, 0.0f); // ...at this rotation rate
-	GetCharacterMovement()->JumpZVelocity = 600.f;
+	GetCharacterMovement()->RotationRate = FRotator(0.0f, 500.0f, 0.0f); // ...at this rotation rate
+
+	// Note: For faster iteration times these variables, and many more, can be tweaked in the Character Blueprint
+	// instead of recompiling to adjust them
+	GetCharacterMovement()->JumpZVelocity = 700.f;
 	GetCharacterMovement()->AirControl = 0.35f;
+	GetCharacterMovement()->MaxWalkSpeed = 500.f;
+	GetCharacterMovement()->MinAnalogWalkSpeed = 20.f;
+	GetCharacterMovement()->BrakingDecelerationWalking = 2000.f;
 
 	// Create a camera boom (pulls in towards the player if there is a collision)
 	CameraBoom = CreateDefaultSubobject<USpringArmComponent>(TEXT("CameraBoom"));
@@ -74,24 +79,24 @@ void ATP_ThirdPersonCharacter::SetupPlayerInputComponent(class UInputComponent*
 
 void ATP_ThirdPersonCharacter::TouchStarted(ETouchIndex::Type FingerIndex, FVector Location)
 {
-		Jump();
+	Jump();
 }
 
 void ATP_ThirdPersonCharacter::TouchStopped(ETouchIndex::Type FingerIndex, FVector Location)
 {
-		StopJumping();
+	StopJumping();
 }
 
 void ATP_ThirdPersonCharacter::TurnAtRate(float Rate)
 {
 	// calculate delta for this frame from the rate information
-	AddControllerYawInput(Rate * BaseTurnRate * GetWorld()->GetDeltaSeconds());
+	AddControllerYawInput(Rate * TurnRateGamepad * GetWorld()->GetDeltaSeconds());
 }
 
 void ATP_ThirdPersonCharacter::LookUpAtRate(float Rate)
 {
 	// calculate delta for this frame from the rate information
-	AddControllerPitchInput(Rate * BaseLookUpRate * GetWorld()->GetDeltaSeconds());
+	AddControllerPitchInput(Rate * TurnRateGamepad * GetWorld()->GetDeltaSeconds());
 }
 
 void ATP_ThirdPersonCharacter::MoveForward(float Value)
diff --git a/Templates/TP_ThirdPerson/Source/TP_ThirdPerson/TP_ThirdPersonCharacter.h b/Templates/TP_ThirdPerson/Source/TP_ThirdPerson/TP_ThirdPersonCharacter.h
--- a/Templates/TP_ThirdPerson/Source/TP_ThirdPerson/TP_ThirdPersonCharacter.h
+++ b/Templates/TP_ThirdPerson/Source/TP_ThirdPerson/TP_ThirdPersonCharacter.h
@@ -22,12 +22,8 @@ public:
 	ATP_ThirdPersonCharacter();
 
 	/** Base turn rate, in deg/sec. Other scaling may affect final turn rate. */
-	UPROPERTY(VisibleAnywhere, BlueprintReadOnly, Category=Camera)
-	float BaseTurnRate;
-
-	/** Base look up/down rate, in deg/sec. Other scaling may affect final rate. */
-	UPROPERTY(VisibleAnywhere, BlueprintReadOnly, Category=Camera)
-	float BaseLookUpRate;
+	UPROPERTY(VisibleAnywhere, BlueprintReadOnly, Category=Input)
+	float TurnRateGamepad;
 
 protected:
 
diff --git a/Templates/TP_ThirdPerson/Source/TP_ThirdPerson/TP_ThirdPersonGameMode.cpp b/Templates/TP_ThirdPerson/Source/TP_ThirdPerson/TP_ThirdPersonGameMode.cpp
--- a/Templates/TP_ThirdPerson/Source/TP_ThirdPerson/TP_ThirdPersonGameMode.cpp
+++ b/Templates/TP_ThirdPerson/Source/TP_ThirdPerson/TP_ThirdPersonGameMode.cpp
@@ -7,7 +7,7 @@
 ATP_ThirdPersonGameMode::ATP_ThirdPersonGameMode()
 {
 	// set default pawn class to our Blueprinted character
-	static ConstructorHelpers::FClassFinder<APawn> PlayerPawnBPClass(TEXT("/Game/ThirdPerson/Blueprints/ThirdPersonCharacter"));
+	static ConstructorHelpers::FClassFinder<APawn> PlayerPawnBPClass(TEXT("/Game/ThirdPerson/Blueprints/BP_ThirdPersonCharacter"));
 	if (PlayerPawnBPClass.Class != NULL)
 	{
 		DefaultPawnClass = PlayerPawnBPClass.Class;
diff --git a/Templates/TP_ThirdPerson/TP_ThirdPerson.uproject b/Templates/TP_ThirdPerson/TP_ThirdPerson.uproject
--- a/Templates/TP_ThirdPerson/TP_ThirdPerson.uproject
+++ b/Templates/TP_ThirdPerson/TP_ThirdPerson.uproject
@@ -18,14 +18,6 @@
 				"Editor"
 			]
 		},
-		{
-			"Name": "OculusVR",
-			"Enabled": false,
-			"SupportedTargetPlatforms": [
-				"Win64",
-				"Android"
-			]
-		},
 		{
 			"Name": "OpenXREyeTracker",
 			"Enabled": false,
@@ -45,24 +37,6 @@
 				"HoloLens",
 				"Android"
 			]
-		},
-		{
-			"Name": "OpenXR",
-			"Enabled": false,
-			"SupportedTargetPlatforms": [
-				"Win64",
-				"Linux",
-				"Android",
-				"HoloLens"
-			]
-		},
-		{
-			"Name": "SteamVR",
-			"Enabled": false,
-			"SupportedTargetPlatforms": [
-				"Win64",
-				"Linux"
-			]
 		}
 	]
 }
\ No newline at end of file
diff --git a/Templates/TP_ThirdPersonBP/Config/DefaultEngine.ini b/Templates/TP_ThirdPersonBP/Config/DefaultEngine.ini
--- a/Templates/TP_ThirdPersonBP/Config/DefaultEngine.ini
+++ b/Templates/TP_ThirdPersonBP/Config/DefaultEngine.ini
@@ -19,4 +19,5 @@ r.ReflectionMethod=1
 r.GenerateMeshDistanceFields=True
 r.DynamicGlobalIlluminationMethod=1
 r.Lumen.TraceMeshSDFs=0
+r.Shadow.Virtual.Enable=1
 
diff --git a/Templates/TP_ThirdPersonBP/TP_ThirdPersonBP.uproject b/Templates/TP_ThirdPersonBP/TP_ThirdPersonBP.uproject
--- a/Templates/TP_ThirdPersonBP/TP_ThirdPersonBP.uproject
+++ b/Templates/TP_ThirdPersonBP/TP_ThirdPersonBP.uproject
@@ -11,14 +11,6 @@
 				"Editor"
 			]
 		},
-		{
-			"Name": "OculusVR",
-			"Enabled": false,
-			"SupportedTargetPlatforms": [
-				"Win64",
-				"Android"
-			]
-		},
 		{
 			"Name": "OpenXREyeTracker",
 			"Enabled": false,
@@ -38,24 +30,6 @@
 				"HoloLens",
 				"Android"
 			]
-		},
-		{
-			"Name": "OpenXR",
-			"Enabled": false,
-			"SupportedTargetPlatforms": [
-				"Win64",
-				"Linux",
-				"Android",
-				"HoloLens"
-			]
-		},
-		{
-			"Name": "SteamVR",
-			"Enabled": false,
-			"SupportedTargetPlatforms": [
-				"Win64",
-				"Linux"
-			]
 		}
 	]
 }
\ No newline at end of file
diff --git a/Templates/TP_TopDown/Config/TemplateDefs.ini b/Templates/TP_TopDown/Config/TemplateDefs.ini
--- a/Templates/TP_TopDown/Config/TemplateDefs.ini
+++ b/Templates/TP_TopDown/Config/TemplateDefs.ini
@@ -47,6 +47,7 @@ FilenameReplacements=(Extensions=("cpp","h","ini","cs"),From="%TEMPLATENAME%",To
 ReplacementsInFiles=(Extensions=("cpp","h","ini","cs"),From="%TEMPLATENAME_UPPERCASE%",To="%PROJECTNAME_UPPERCASE%",bCaseSensitive=true)
 ReplacementsInFiles=(Extensions=("cpp","h","ini","cs"),From="%TEMPLATENAME_LOWERCASE%",To="%PROJECTNAME_LOWERCASE%",bCaseSensitive=true)
 ReplacementsInFiles=(Extensions=("cpp","h","ini","cs"),From="%TEMPLATENAME%",To="%PROJECTNAME%",bCaseSensitive=false)SharedContentPacks=(MountName="TopDown",DetailLevels=("Standard"))
-SharedContentPacks=(MountName="LevelPrototyping",DetailLevels=("Standard","High"))
-SharedContentPacks=(MountName="Mannequin",DetailLevels=("Standard","High"))
+SharedContentPacks=(MountName="LevelPrototyping",DetailLevels=("High"))
+SharedContentPacks=(MountName="Mannequin",DetailLevels=("High"))
+SharedContentPacks=(MountName="Cursor",DetailLevels=("Standard"))
 EditDetailLevelPreference="High"
\ No newline at end of file
diff --git a/Templates/TP_TopDown/Source/TP_TopDown/TP_TopDown.Build.cs b/Templates/TP_TopDown/Source/TP_TopDown/TP_TopDown.Build.cs
--- a/Templates/TP_TopDown/Source/TP_TopDown/TP_TopDown.Build.cs
+++ b/Templates/TP_TopDown/Source/TP_TopDown/TP_TopDown.Build.cs
@@ -8,6 +8,6 @@ public class TP_TopDown : ModuleRules
 	{
 		PCHUsage = PCHUsageMode.UseExplicitOrSharedPCHs;
 
-        PublicDependencyModuleNames.AddRange(new string[] { "Core", "CoreUObject", "Engine", "InputCore", "HeadMountedDisplay", "NavigationSystem", "AIModule" });
+        PublicDependencyModuleNames.AddRange(new string[] { "Core", "CoreUObject", "Engine", "InputCore", "HeadMountedDisplay", "NavigationSystem", "AIModule", "Niagara" });
     }
 }
diff --git a/Templates/TP_TopDown/Source/TP_TopDown/TP_TopDownCharacter.cpp b/Templates/TP_TopDown/Source/TP_TopDown/TP_TopDownCharacter.cpp
--- a/Templates/TP_TopDown/Source/TP_TopDown/TP_TopDownCharacter.cpp
+++ b/Templates/TP_TopDown/Source/TP_TopDown/TP_TopDownCharacter.cpp
@@ -40,17 +40,6 @@ ATP_TopDownCharacter::ATP_TopDownCharacter()
 	TopDownCameraComponent->SetupAttachment(CameraBoom, USpringArmComponent::SocketName);
 	TopDownCameraComponent->bUsePawnControlRotation = false; // Camera does not rotate relative to arm
 
-	// Create a decal in the world to show the cursor's location
-	CursorToWorld = CreateDefaultSubobject<UDecalComponent>("CursorToWorld");
-	CursorToWorld->SetupAttachment(RootComponent);
-	static ConstructorHelpers::FObjectFinder<UMaterial> DecalMaterialAsset(TEXT("Material'/Game/TopDown/Blueprints/M_Cursor_Decal.M_Cursor_Decal'"));
-	if (DecalMaterialAsset.Succeeded())
-	{
-		CursorToWorld->SetDecalMaterial(DecalMaterialAsset.Object);
-	}
-	CursorToWorld->DecalSize = FVector(16.0f, 32.0f, 32.0f);
-	CursorToWorld->SetRelativeRotation(FRotator(90.0f, 0.0f, 0.0f).Quaternion());
-
 	// Activate ticking in order to update the cursor every frame.
 	PrimaryActorTick.bCanEverTick = true;
 	PrimaryActorTick.bStartWithTickEnabled = true;
@@ -59,17 +48,4 @@ ATP_TopDownCharacter::ATP_TopDownCharacter()
 void ATP_TopDownCharacter::Tick(float DeltaSeconds)
 {
     Super::Tick(DeltaSeconds);
-
-	if (CursorToWorld != nullptr)
-	{
-		if (APlayerController* PC = Cast<APlayerController>(GetController()))
-		{
-			FHitResult TraceHitResult;
-			PC->GetHitResultUnderCursor(ECC_Visibility, true, TraceHitResult);
-			FVector CursorFV = TraceHitResult.ImpactNormal;
-			FRotator CursorR = CursorFV.Rotation();
-			CursorToWorld->SetWorldLocation(TraceHitResult.Location);
-			CursorToWorld->SetWorldRotation(CursorR);
-		}
-	}
 }
diff --git a/Templates/TP_TopDown/Source/TP_TopDown/TP_TopDownCharacter.h b/Templates/TP_TopDown/Source/TP_TopDown/TP_TopDownCharacter.h
--- a/Templates/TP_TopDown/Source/TP_TopDown/TP_TopDownCharacter.h
+++ b/Templates/TP_TopDown/Source/TP_TopDown/TP_TopDownCharacter.h
@@ -21,8 +21,6 @@ public:
 	FORCEINLINE class UCameraComponent* GetTopDownCameraComponent() const { return TopDownCameraComponent; }
 	/** Returns CameraBoom subobject **/
 	FORCEINLINE class USpringArmComponent* GetCameraBoom() const { return CameraBoom; }
-	/** Returns CursorToWorld subobject **/
-	FORCEINLINE class UDecalComponent* GetCursorToWorld() { return CursorToWorld; }
 
 private:
 	/** Top down camera */
@@ -32,9 +30,5 @@ private:
 	/** Camera boom positioning the camera above the character */
 	UPROPERTY(VisibleAnywhere, BlueprintReadOnly, Category = Camera, meta = (AllowPrivateAccess = "true"))
 	class USpringArmComponent* CameraBoom;
-
-	/** A decal that projects to the cursor location. */
-	UPROPERTY(VisibleAnywhere, BlueprintReadOnly, Category = Camera, meta = (AllowPrivateAccess = "true"))
-	class UDecalComponent* CursorToWorld;
 };
 
diff --git a/Templates/TP_TopDown/Source/TP_TopDown/TP_TopDownGameMode.cpp b/Templates/TP_TopDown/Source/TP_TopDown/TP_TopDownGameMode.cpp
--- a/Templates/TP_TopDown/Source/TP_TopDown/TP_TopDownGameMode.cpp
+++ b/Templates/TP_TopDown/Source/TP_TopDown/TP_TopDownGameMode.cpp
@@ -16,4 +16,11 @@ ATP_TopDownGameMode::ATP_TopDownGameMode()
 	{
 		DefaultPawnClass = PlayerPawnBPClass.Class;
 	}
+
+	// set default controller to our Blueprinted controller
+	static ConstructorHelpers::FClassFinder<APlayerController> PlayerControllerBPClass(TEXT("/Game/TopDown/Blueprints/BP_TopDownPlayerController"));
+	if(PlayerControllerBPClass.Class != NULL)
+	{
+		PlayerControllerClass = PlayerControllerBPClass.Class;
+	}
 }
\ No newline at end of file
diff --git a/Templates/TP_TopDown/Source/TP_TopDown/TP_TopDownPlayerController.cpp b/Templates/TP_TopDown/Source/TP_TopDown/TP_TopDownPlayerController.cpp
--- a/Templates/TP_TopDown/Source/TP_TopDown/TP_TopDownPlayerController.cpp
+++ b/Templates/TP_TopDown/Source/TP_TopDown/TP_TopDownPlayerController.cpp
@@ -1,25 +1,51 @@
 // Copyright Epic Games, Inc. All Rights Reserved.
 
 #include "TP_TopDownPlayerController.h"
+#include "GameFramework/Pawn.h"
 #include "Blueprint/AIBlueprintHelperLibrary.h"
-#include "Runtime/Engine/Classes/Components/DecalComponent.h"
+#include "NiagaraSystem.h"
+#include "NiagaraFunctionLibrary.h"
 #include "TP_TopDownCharacter.h"
 #include "Engine/World.h"
 
 ATP_TopDownPlayerController::ATP_TopDownPlayerController()
 {
 	bShowMouseCursor = true;
-	DefaultMouseCursor = EMouseCursor::Crosshairs;
+	DefaultMouseCursor = EMouseCursor::Default;
 }
 
 void ATP_TopDownPlayerController::PlayerTick(float DeltaTime)
 {
 	Super::PlayerTick(DeltaTime);
 
-	// keep updating the destination every tick while desired
-	if (bMoveToMouseCursor)
+	if(bInputPressed)
 	{
-		MoveToMouseCursor();
+		FollowTime += DeltaTime;
+
+		// Look for the touch location
+		FVector HitLocation = FVector::ZeroVector;
+		FHitResult Hit;
+		if(bIsTouch)
+		{
+			GetHitResultUnderFinger(ETouchIndex::Touch1, ECC_Visibility, true, Hit);
+		}
+		else
+		{
+			GetHitResultUnderCursor(ECC_Visibility, true, Hit);
+		}
+		HitLocation = Hit.Location;
+
+		// Direct the Pawn towards that location
+		APawn* const MyPawn = GetPawn();
+		if(MyPawn)
+		{
+			FVector WorldDirection = (HitLocation - MyPawn->GetActorLocation()).GetSafeNormal();
+			MyPawn->AddMovementInput(WorldDirection, 1.f, false);
+		}
+	}
+	else
+	{
+		FollowTime = 0.f;
 	}
 }
 
@@ -32,64 +58,47 @@ void ATP_TopDownPlayerController::SetupInputComponent()
 	InputComponent->BindAction("SetDestination", IE_Released, this, &ATP_TopDownPlayerController::OnSetDestinationReleased);
 
 	// support touch devices 
-	InputComponent->BindTouch(EInputEvent::IE_Pressed, this, &ATP_TopDownPlayerController::MoveToTouchLocation);
-	InputComponent->BindTouch(EInputEvent::IE_Repeat, this, &ATP_TopDownPlayerController::MoveToTouchLocation);
+	InputComponent->BindTouch(EInputEvent::IE_Pressed, this, &ATP_TopDownPlayerController::OnTouchPressed);
+	InputComponent->BindTouch(EInputEvent::IE_Released, this, &ATP_TopDownPlayerController::OnTouchReleased);
 
 }
 
-void ATP_TopDownPlayerController::MoveToMouseCursor()
+void ATP_TopDownPlayerController::OnSetDestinationPressed()
 {
-	
-	{
-		// Trace to see what is under the mouse cursor
-		FHitResult Hit;
-		GetHitResultUnderCursor(ECC_Visibility, false, Hit);
-
-		if (Hit.bBlockingHit)
-		{
-			// We hit something, move there
-			SetNewMoveDestination(Hit.ImpactPoint);
-		}
-	}
+	// We flag that the input is being pressed
+	bInputPressed = true;
+	// Just in case the character was moving because of a previous short press we stop it
+	StopMovement();
 }
 
-void ATP_TopDownPlayerController::MoveToTouchLocation(const ETouchIndex::Type FingerIndex, const FVector Location)
+void ATP_TopDownPlayerController::OnSetDestinationReleased()
 {
-	FVector2D ScreenSpaceLocation(Location);
-
-	// Trace to see what is under the touch location
-	FHitResult HitResult;
-	GetHitResultAtScreenPosition(ScreenSpaceLocation, CurrentClickTraceChannel, true, HitResult);
-	if (HitResult.bBlockingHit)
-	{
-		// We hit something, move there
-		SetNewMoveDestination(HitResult.ImpactPoint);
-	}
-}
+	// Player is no longer pressing the input
+	bInputPressed = false;
 
-void ATP_TopDownPlayerController::SetNewMoveDestination(const FVector DestLocation)
-{
-	APawn* const MyPawn = GetPawn();
-	if (MyPawn)
+	// If it was a short press
+	if(FollowTime <= ShortPressThreshold)
 	{
-		float const Distance = FVector::Dist(DestLocation, MyPawn->GetActorLocation());
+		// We look for the location in the world where the player has pressed the input
+		FVector HitLocation = FVector::ZeroVector;
+		FHitResult Hit;
+		GetHitResultUnderCursor(ECC_Visibility, true, Hit);
+		HitLocation = Hit.Location;
 
-		// We need to issue move command only if far enough in order for walk animation to play correctly
-		if ((Distance > 120.0f))
-		{
-			UAIBlueprintHelperLibrary::SimpleMoveToLocation(this, DestLocation);
-		}
+		// We move there and spawn some particles
+		UAIBlueprintHelperLibrary::SimpleMoveToLocation(this, HitLocation);
+		UNiagaraFunctionLibrary::SpawnSystemAtLocation(this, FXCursor, HitLocation, FRotator::ZeroRotator, FVector(1.f, 1.f, 1.f), true, true, ENCPoolMethod::None, true);
 	}
 }
 
-void ATP_TopDownPlayerController::OnSetDestinationPressed()
+void ATP_TopDownPlayerController::OnTouchPressed(const ETouchIndex::Type FingerIndex, const FVector Location)
 {
-	// set flag to keep updating destination until released
-	bMoveToMouseCursor = true;
+	bIsTouch = true;
+	OnSetDestinationPressed();
 }
 
-void ATP_TopDownPlayerController::OnSetDestinationReleased()
+void ATP_TopDownPlayerController::OnTouchReleased(const ETouchIndex::Type FingerIndex, const FVector Location)
 {
-	// clear flag to indicate we should stop updating the destination
-	bMoveToMouseCursor = false;
+	bIsTouch = false;
+	OnSetDestinationReleased();
 }
diff --git a/Templates/TP_TopDown/Source/TP_TopDown/TP_TopDownPlayerController.h b/Templates/TP_TopDown/Source/TP_TopDown/TP_TopDownPlayerController.h
--- a/Templates/TP_TopDown/Source/TP_TopDown/TP_TopDownPlayerController.h
+++ b/Templates/TP_TopDown/Source/TP_TopDown/TP_TopDownPlayerController.h
@@ -3,9 +3,13 @@
 #pragma once
 
 #include "CoreMinimal.h"
+#include "Templates/SubclassOf.h"
 #include "GameFramework/PlayerController.h"
 #include "TP_TopDownPlayerController.generated.h"
 
+/** Forward declaration to improve compiling times */
+class UNiagaraSystem;
+
 UCLASS()
 class ATP_TopDownPlayerController : public APlayerController
 {
@@ -14,6 +18,14 @@ class ATP_TopDownPlayerController : public APlayerController
 public:
 	ATP_TopDownPlayerController();
 
+	/** Time Threshold to know if it was a short press */
+	UPROPERTY(EditAnywhere, BlueprintReadOnly, Category = Input)
+	float ShortPressThreshold;
+
+	/** FX Class that we will spawn when clicking */
+	UPROPERTY(EditAnywhere, BlueprintReadOnly, Category = Input)
+	UNiagaraSystem* FXCursor;
+
 protected:
 	/** True if the controlled character should navigate to the mouse cursor. */
 	uint32 bMoveToMouseCursor : 1;
@@ -23,18 +35,16 @@ protected:
 	virtual void SetupInputComponent() override;
 	// End PlayerController interface
 
-	/** Navigate player to the current mouse cursor location. */
-	void MoveToMouseCursor();
-
-	/** Navigate player to the current touch location. */
-	void MoveToTouchLocation(const ETouchIndex::Type FingerIndex, const FVector Location);
-	
-	/** Navigate player to the given world location. */
-	void SetNewMoveDestination(const FVector DestLocation);
-
 	/** Input handlers for SetDestination action. */
 	void OnSetDestinationPressed();
 	void OnSetDestinationReleased();
+	void OnTouchPressed(const ETouchIndex::Type FingerIndex, const FVector Location);
+	void OnTouchReleased(const ETouchIndex::Type FingerIndex, const FVector Location);
+
+private:
+	bool bInputPressed; // Input is bring pressed
+	bool bIsTouch; // Is it a touch device
+	float FollowTime; // For how long it has been pressed
 };
 
 
diff --git a/Templates/TP_TopDown/manifest.json b/Templates/TP_TopDown/manifest.json
--- a/Templates/TP_TopDown/manifest.json
+++ b/Templates/TP_TopDown/manifest.json
@@ -69,6 +69,13 @@
 				"High"
 			]
 		},				
+{
+			"MountName": "Cursor",
+			"DetailLevels":
+			[
+				"Standard"
+			]
+		},
 		{
 			"MountName": "Mannequin",
 			"DetailLevels":
diff --git a/Templates/TP_TopDownBP/Config/DefaultEngine.ini b/Templates/TP_TopDownBP/Config/DefaultEngine.ini
--- a/Templates/TP_TopDownBP/Config/DefaultEngine.ini
+++ b/Templates/TP_TopDownBP/Config/DefaultEngine.ini
@@ -52,3 +52,9 @@ bForceRebuildOnLoad=True
 [/Script/IOSRuntimeSettings.IOSRuntimeSettings]
 MinimumiOSVersion=IOS_13
 
+[/Script/Engine.RendererSettings]
+r.ReflectionMethod=1
+r.GenerateMeshDistanceFields=True
+r.DynamicGlobalIlluminationMethod=1
+r.Shadow.Virtual.Enable=1
+
diff --git a/Templates/TP_TopDownBP/Config/TemplateDefs.ini b/Templates/TP_TopDownBP/Config/TemplateDefs.ini
--- a/Templates/TP_TopDownBP/Config/TemplateDefs.ini
+++ b/Templates/TP_TopDownBP/Config/TemplateDefs.ini
@@ -42,6 +42,7 @@ FilenameReplacements=(Extensions=("cpp","h","ini","cs"),From="%TEMPLATENAME%",To
 ReplacementsInFiles=(Extensions=("cpp","h","ini","cs"),From="%TEMPLATENAME_UPPERCASE%",To="%PROJECTNAME_UPPERCASE%",bCaseSensitive=true)
 ReplacementsInFiles=(Extensions=("cpp","h","ini","cs"),From="%TEMPLATENAME_LOWERCASE%",To="%PROJECTNAME_LOWERCASE%",bCaseSensitive=true)
 ReplacementsInFiles=(Extensions=("cpp","h","ini","cs"),From="%TEMPLATENAME%",To="%PROJECTNAME%",bCaseSensitive=false)
-SharedContentPacks=(MountName="LevelPrototyping",DetailLevels=("Standard","High"))
-SharedContentPacks=(MountName="Mannequin",DetailLevels=("Standard","High"))
+SharedContentPacks=(MountName="LevelPrototyping",DetailLevels=("High"))
+SharedContentPacks=(MountName="Mannequin",DetailLevels=("High"))
+SharedContentPacks=(MountName="Cursor",DetailLevels=("Standard"))
 EditDetailLevelPreference="High"
diff --git a/Templates/TP_TopDownBP/TP_TopDownBP.uproject b/Templates/TP_TopDownBP/TP_TopDownBP.uproject
--- a/Templates/TP_TopDownBP/TP_TopDownBP.uproject
+++ b/Templates/TP_TopDownBP/TP_TopDownBP.uproject
@@ -10,52 +10,6 @@
 			"TargetAllowList": [
 				"Editor"
 			]
-		},
-		{
-			"Name": "OculusVR",
-			"Enabled": false,
-			"SupportedTargetPlatforms": [
-				"Win64",
-				"Android"
-			]
-		},
-		{
-			"Name": "OpenXREyeTracker",
-			"Enabled": false,
-			"SupportedTargetPlatforms": [
-				"Win64",
-				"Linux",
-				"HoloLens",
-				"Android"
-			]
-		},
-		{
-			"Name": "OpenXRHandTracking",
-			"Enabled": false,
-			"SupportedTargetPlatforms": [
-				"Win64",
-				"Linux",
-				"HoloLens",
-				"Android"
-			]
-		},
-		{
-			"Name": "OpenXR",
-			"Enabled": false,
-			"SupportedTargetPlatforms": [
-				"Win64",
-				"Linux",
-				"Android",
-				"HoloLens"
-			]
-		},
-		{
-			"Name": "SteamVR",
-			"Enabled": false,
-			"SupportedTargetPlatforms": [
-				"Win64",
-				"Linux"
-			]
 		}
 	]
 }
\ No newline at end of file
diff --git a/Templates/TP_TopDownBP/manifest.json b/Templates/TP_TopDownBP/manifest.json
--- a/Templates/TP_TopDownBP/manifest.json
+++ b/Templates/TP_TopDownBP/manifest.json
@@ -69,6 +69,13 @@
 				"High"
 			]
 		},
+		{
+			"MountName": "Cursor",
+			"DetailLevels":
+			[			
+				"Standard"
+			]
+		},
 		{
 			"MountName": "Mannequin",
 			"DetailLevels":
