Templates - TP_TopDown
Updated template source code and comments.
Removed source code asset references.
Implemented Strategy and TwinStick variants.
#rnx

diff --git a/Templates/TP_TopDown/Config/DefaultEngine.ini b/Templates/TP_TopDown/Config/DefaultEngine.ini
--- a/Templates/TP_TopDown/Config/DefaultEngine.ini
+++ b/Templates/TP_TopDown/Config/DefaultEngine.ini
@@ -1,7 +1,7 @@
 [/Script/EngineSettings.GameMapsSettings]
-GameDefaultMap=/Game/TopDown/Maps/TopDownMap.TopDownMap
-EditorStartupMap=/Game/TopDown/Maps/TopDownMap.TopDownMap
-GlobalDefaultGameMode="/Script/TP_TopDown.TP_TopDownGameMode"
+GameDefaultMap=/Game/TopDown/Lvl_TopDown.Lvl_TopDown
+EditorStartupMap=/Game/TopDown/Lvl_TopDown.Lvl_TopDown
+GlobalDefaultGameMode=/Game/TopDown/Blueprints/BP_TopDownGameMode.BP_TopDownGameMode_C
 
 [/Script/Engine.RendererSettings]
 r.DefaultFeature.AutoExposure.ExtendDefaultLuminanceRange=true
@@ -39,3 +39,44 @@ SoundCueCookQualityIndex=-1
 [/Script/LinuxTargetPlatform.LinuxTargetSettings]
 -TargetedRHIs=SF_VULKAN_SM5
 +TargetedRHIs=SF_VULKAN_SM6
+
+[CoreRedirects]
++ClassRedirects=(OldName="/Script/TopDownCPP.StrategyGameMode",NewName="/Script/TP_TopDown.StrategyGameMode")
++ClassRedirects=(OldName="/Script/TopDownCPP.StrategyHUD",NewName="/Script/TP_TopDown.StrategyHUD")
++ClassRedirects=(OldName="/Script/TopDownCPP.StrategyPawn",NewName="/Script/TP_TopDown.StrategyPawn")
++ClassRedirects=(OldName="/Script/TopDownCPP.StrategyPlayerController",NewName="/Script/TP_TopDown.StrategyPlayerController")
++ClassRedirects=(OldName="/Script/TopDownCPP.StrategyUI",NewName="/Script/TP_TopDown.StrategyUI")
++ClassRedirects=(OldName="/Script/TopDownCPP.StrategyNPC",NewName="/Script/TP_TopDown.StrategyUnit")
+
++ClassRedirects=(OldName="/Script/TopDownCPP.TwinStickAIController",NewName="/Script/TP_TopDown.TwinStickAIController")
++ClassRedirects=(OldName="/Script/TopDownCPP.TwinStickNPC",NewName="/Script/TP_TopDown.TwinStickNPC")
++ClassRedirects=(OldName="/Script/TopDownCPP.TwinStickAoEAttack",NewName="/Script/TP_TopDown.TwinStickAoEAttack")
++ClassRedirects=(OldName="/Script/TopDownCPP.TwinStickCharacter",NewName="/Script/TP_TopDown.TwinStickCharacter")
++ClassRedirects=(OldName="/Script/TopDownCPP.TwinStickGameMode",NewName="/Script/TP_TopDown.TwinStickGameMode")
++ClassRedirects=(OldName="/Script/TopDownCPP.TwinStickNPCDestruction",NewName="/Script/TP_TopDown.TwinStickNPCDestruction")
++ClassRedirects=(OldName="/Script/TopDownCPP.TwinStickPickup",NewName="/Script/TP_TopDown.TwinStickPickup")
++ClassRedirects=(OldName="/Script/TopDownCPP.TwinStickPlayerController",NewName="/Script/TP_TopDown.TwinStickPlayerController")
++ClassRedirects=(OldName="/Script/TopDownCPP.TwinStickProjectile",NewName="/Script/TP_TopDown.TwinStickProjectile")
++ClassRedirects=(OldName="/Script/TopDownCPP.TwinStickSpawner",NewName="/Script/TP_TopDown.TwinStickSpawner")
++ClassRedirects=(OldName="/Script/TopDownCPP.TwinStickUI",NewName="/Script/TP_TopDown.TwinStickUI")
+
+;+StructRedirects=(OldName="MyStruct",NewName="MyNewStruct")
+;+StructRedirects=(OldName="MyStruct",NewName="MyNewStruct")
+
+[/Script/AndroidFileServerEditor.AndroidFileServerRuntimeSettings]
+bEnablePlugin=True
+bAllowNetworkConnection=True
+SecurityToken=0FB14AED42FEC3066D998A8660BACCA3
+bIncludeInShipping=False
+bAllowExternalStartInShipping=False
+bCompileAFSProject=False
+bUseCompression=False
+bLogFiles=False
+bReportStats=False
+ConnectionType=USBOnly
+bUseManualIPAddress=False
+ManualIPAddress=
+
+[/Script/NavigationSystem.RecastNavMesh]
+RuntimeGeneration=Dynamic
+
diff --git a/Templates/TP_TopDown/Source/TP_TopDown/TP_TopDown.Build.cs b/Templates/TP_TopDown/Source/TP_TopDown/TP_TopDown.Build.cs
--- a/Templates/TP_TopDown/Source/TP_TopDown/TP_TopDown.Build.cs
+++ b/Templates/TP_TopDown/Source/TP_TopDown/TP_TopDown.Build.cs
@@ -8,6 +8,35 @@ public class TP_TopDown : ModuleRules
 	{
 		PCHUsage = PCHUsageMode.UseExplicitOrSharedPCHs;
 
-        PublicDependencyModuleNames.AddRange(new string[] { "Core", "CoreUObject", "Engine", "InputCore", "NavigationSystem", "AIModule", "Niagara", "EnhancedInput" });
-    }
+		PublicDependencyModuleNames.AddRange(new string[] {
+			"Core",
+			"CoreUObject",
+			"Engine",
+			"InputCore",
+			"EnhancedInput",
+			"AIModule",
+			"NavigationSystem",
+			"StateTreeModule",
+			"GameplayStateTreeModule",
+			"Niagara",
+			"UMG"
+		});
+
+		PrivateDependencyModuleNames.AddRange(new string[] { });
+
+		PublicIncludePaths.AddRange(new string[] {
+			"TP_TopDown",
+			"TP_TopDown/Variant_Strategy",
+			"TP_TopDown/Variant_TwinStick",
+			"TP_TopDown/Variant_TwinStick/AI"
+		});
+
+		// Uncomment if you are using Slate UI
+		// PrivateDependencyModuleNames.AddRange(new string[] { "Slate", "SlateCore" });
+
+		// Uncomment if you are using online features
+		// PrivateDependencyModuleNames.Add("OnlineSubsystem");
+
+		// To include OnlineSubsystemSteam, add it to the plugins section in your uproject file with the Enabled attribute set to true
+	}
 }
diff --git a/Templates/TP_TopDown/Source/TP_TopDown/TP_TopDownCharacter.cpp b/Templates/TP_TopDown/Source/TP_TopDown/TP_TopDownCharacter.cpp
--- a/Templates/TP_TopDown/Source/TP_TopDown/TP_TopDownCharacter.cpp
+++ b/Templates/TP_TopDown/Source/TP_TopDown/TP_TopDownCharacter.cpp
@@ -22,30 +22,41 @@ ATP_TopDownCharacter::ATP_TopDownCharacter()
 	bUseControllerRotationRoll = false;
 
 	// Configure character movement
-	GetCharacterMovement()->bOrientRotationToMovement = true; // Rotate character to moving direction
+	GetCharacterMovement()->bOrientRotationToMovement = true;
 	GetCharacterMovement()->RotationRate = FRotator(0.f, 640.f, 0.f);
 	GetCharacterMovement()->bConstrainToPlane = true;
 	GetCharacterMovement()->bSnapToPlaneAtStart = true;
 
-	// Create a camera boom...
+	// Create the camera boom component
 	CameraBoom = CreateDefaultSubobject<USpringArmComponent>(TEXT("CameraBoom"));
+
 	CameraBoom->SetupAttachment(RootComponent);
-	CameraBoom->SetUsingAbsoluteRotation(true); // Don't want arm to rotate when character does
+	CameraBoom->SetUsingAbsoluteRotation(true);
 	CameraBoom->TargetArmLength = 800.f;
 	CameraBoom->SetRelativeRotation(FRotator(-60.f, 0.f, 0.f));
-	CameraBoom->bDoCollisionTest = false; // Don't want to pull camera in when it collides with level
+	CameraBoom->bDoCollisionTest = false;
 
-	// Create a camera...
+	// Create the camera component
 	TopDownCameraComponent = CreateDefaultSubobject<UCameraComponent>(TEXT("TopDownCamera"));
+
 	TopDownCameraComponent->SetupAttachment(CameraBoom, USpringArmComponent::SocketName);
-	TopDownCameraComponent->bUsePawnControlRotation = false; // Camera does not rotate relative to arm
+	TopDownCameraComponent->bUsePawnControlRotation = false;
 
 	// Activate ticking in order to update the cursor every frame.
 	PrimaryActorTick.bCanEverTick = true;
 	PrimaryActorTick.bStartWithTickEnabled = true;
 }
 
+void ATP_TopDownCharacter::BeginPlay()
+{
+	Super::BeginPlay();
+
+	// stub
+}
+
 void ATP_TopDownCharacter::Tick(float DeltaSeconds)
 {
     Super::Tick(DeltaSeconds);
+
+	// stub
 }
diff --git a/Templates/TP_TopDown/Source/TP_TopDown/TP_TopDownCharacter.h b/Templates/TP_TopDown/Source/TP_TopDown/TP_TopDownCharacter.h
--- a/Templates/TP_TopDown/Source/TP_TopDown/TP_TopDownCharacter.h
+++ b/Templates/TP_TopDown/Source/TP_TopDown/TP_TopDownCharacter.h
@@ -6,23 +6,16 @@
 #include "GameFramework/Character.h"
 #include "TP_TopDownCharacter.generated.h"
 
-UCLASS(Blueprintable)
+/**
+ *  A controllable top-down perspective character
+ */
+UCLASS(abstract)
 class ATP_TopDownCharacter : public ACharacter
 {
 	GENERATED_BODY()
 
-public:
-	ATP_TopDownCharacter();
-
-	// Called every frame.
-	virtual void Tick(float DeltaSeconds) override;
-
-	/** Returns TopDownCameraComponent subobject **/
-	FORCEINLINE class UCameraComponent* GetTopDownCameraComponent() const { return TopDownCameraComponent; }
-	/** Returns CameraBoom subobject **/
-	FORCEINLINE class USpringArmComponent* GetCameraBoom() const { return CameraBoom; }
-
 private:
+
 	/** Top down camera */
 	UPROPERTY(VisibleAnywhere, BlueprintReadOnly, Category = Camera, meta = (AllowPrivateAccess = "true"))
 	class UCameraComponent* TopDownCameraComponent;
@@ -30,5 +23,23 @@ private:
 	/** Camera boom positioning the camera above the character */
 	UPROPERTY(VisibleAnywhere, BlueprintReadOnly, Category = Camera, meta = (AllowPrivateAccess = "true"))
 	class USpringArmComponent* CameraBoom;
+
+public:
+
+	/** Constructor */
+	ATP_TopDownCharacter();
+
+	/** Initialization */
+	virtual void BeginPlay() override;
+
+	/** Update */
+	virtual void Tick(float DeltaSeconds) override;
+
+	/** Returns the camera component **/
+	FORCEINLINE class UCameraComponent* GetTopDownCameraComponent() const { return TopDownCameraComponent; }
+
+	/** Returns the Camera Boom component **/
+	FORCEINLINE class USpringArmComponent* GetCameraBoom() const { return CameraBoom; }
+
 };
 
diff --git a/Templates/TP_TopDown/Source/TP_TopDown/TP_TopDownGameMode.cpp b/Templates/TP_TopDown/Source/TP_TopDown/TP_TopDownGameMode.cpp
--- a/Templates/TP_TopDown/Source/TP_TopDown/TP_TopDownGameMode.cpp
+++ b/Templates/TP_TopDown/Source/TP_TopDown/TP_TopDownGameMode.cpp
@@ -1,26 +1,8 @@
 // Copyright Epic Games, Inc. All Rights Reserved.
 
 #include "TP_TopDownGameMode.h"
-#include "TP_TopDownPlayerController.h"
-#include "TP_TopDownCharacter.h"
-#include "UObject/ConstructorHelpers.h"
 
 ATP_TopDownGameMode::ATP_TopDownGameMode()
 {
-	// use our custom PlayerController class
-	PlayerControllerClass = ATP_TopDownPlayerController::StaticClass();
-
-	// set default pawn class to our Blueprinted character
-	static ConstructorHelpers::FClassFinder<APawn> PlayerPawnBPClass(TEXT("/Game/TopDown/Blueprints/BP_TopDownCharacter"));
-	if (PlayerPawnBPClass.Class != nullptr)
-	{
-		DefaultPawnClass = PlayerPawnBPClass.Class;
-	}
-
-	// set default controller to our Blueprinted controller
-	static ConstructorHelpers::FClassFinder<APlayerController> PlayerControllerBPClass(TEXT("/Game/TopDown/Blueprints/BP_TopDownPlayerController"));
-	if(PlayerControllerBPClass.Class != NULL)
-	{
-		PlayerControllerClass = PlayerControllerBPClass.Class;
-	}
+	// stub
 }
\ No newline at end of file
diff --git a/Templates/TP_TopDown/Source/TP_TopDown/TP_TopDownGameMode.h b/Templates/TP_TopDown/Source/TP_TopDown/TP_TopDownGameMode.h
--- a/Templates/TP_TopDown/Source/TP_TopDown/TP_TopDownGameMode.h
+++ b/Templates/TP_TopDown/Source/TP_TopDown/TP_TopDownGameMode.h
@@ -6,12 +6,19 @@
 #include "GameFramework/GameModeBase.h"
 #include "TP_TopDownGameMode.generated.h"
 
-UCLASS(minimalapi)
+/**
+ *  Simple Game Mode for a top-down perspective game
+ *  Sets the default gameplay framework classes
+ *  Check the Blueprint derived class for the set values
+ */
+UCLASS(abstract)
 class ATP_TopDownGameMode : public AGameModeBase
 {
 	GENERATED_BODY()
 
 public:
+
+	/** Constructor */
 	ATP_TopDownGameMode();
 };
 
diff --git a/Templates/TP_TopDown/Source/TP_TopDown/TP_TopDownPlayerController.cpp b/Templates/TP_TopDown/Source/TP_TopDown/TP_TopDownPlayerController.cpp
--- a/Templates/TP_TopDown/Source/TP_TopDown/TP_TopDownPlayerController.cpp
+++ b/Templates/TP_TopDown/Source/TP_TopDown/TP_TopDownPlayerController.cpp
@@ -16,6 +16,10 @@ DEFINE_LOG_CATEGORY(LogTemplateCharacter);
 
 ATP_TopDownPlayerController::ATP_TopDownPlayerController()
 {
+	bIsTouch = false;
+	bMoveToMouseCursor = false;
+
+	// configure the 
 	bShowMouseCursor = true;
 	DefaultMouseCursor = EMouseCursor::Default;
 	CachedDestination = FVector::ZeroVector;
@@ -65,7 +69,6 @@ void ATP_TopDownPlayerController::OnInputStarted()
 	StopMovement();
 }
 
-// Triggered every frame when the input is held down
 void ATP_TopDownPlayerController::OnSetDestinationTriggered()
 {
 	// We flag that the input is being pressed
diff --git a/Templates/TP_TopDown/Source/TP_TopDown/TP_TopDownPlayerController.h b/Templates/TP_TopDown/Source/TP_TopDown/TP_TopDownPlayerController.h
--- a/Templates/TP_TopDown/Source/TP_TopDown/TP_TopDownPlayerController.h
+++ b/Templates/TP_TopDown/Source/TP_TopDown/TP_TopDownPlayerController.h
@@ -7,20 +7,22 @@
 #include "GameFramework/PlayerController.h"
 #include "TP_TopDownPlayerController.generated.h"
 
-/** Forward declaration to improve compiling times */
 class UNiagaraSystem;
 class UInputMappingContext;
 class UInputAction;
 
 DECLARE_LOG_CATEGORY_EXTERN(LogTemplateCharacter, Log, All);
 
-UCLASS()
+/**
+ *  Player controller for a top-down perspective game.
+ *  Implements point and click based controls
+ */
+UCLASS(abstract)
 class ATP_TopDownPlayerController : public APlayerController
 {
 	GENERATED_BODY()
 
-public:
-	ATP_TopDownPlayerController();
+protected:
 
 	/** Time Threshold to know if it was a short press */
 	UPROPERTY(EditAnywhere, BlueprintReadOnly, Category = Input)
@@ -42,27 +44,38 @@ public:
 	UPROPERTY(EditAnywhere, BlueprintReadOnly, Category=Input, meta=(AllowPrivateAccess = "true"))
 	UInputAction* SetDestinationTouchAction;
 
-protected:
 	/** True if the controlled character should navigate to the mouse cursor. */
 	uint32 bMoveToMouseCursor : 1;
 
+	/** Set to true if we're using touch input */
+	uint32 bIsTouch : 1;
+
+	/** Saved location of the character movement destination */
+	FVector CachedDestination;
+
+	/** Time that the click input has been pressed */
+	float FollowTime = 0.0f;
+
+public:
+
+	/** Constructor */
+	ATP_TopDownPlayerController();
+
+protected:
+
+	/** Initialize input bindings */
 	virtual void SetupInputComponent() override;
 	
-	// To add mapping context
+	/** Gameplay initialization */
 	virtual void BeginPlay();
 
-	/** Input handlers for SetDestination action. */
+	/** Input handlers */
 	void OnInputStarted();
 	void OnSetDestinationTriggered();
 	void OnSetDestinationReleased();
 	void OnTouchTriggered();
 	void OnTouchReleased();
 
-private:
-	FVector CachedDestination;
-
-	bool bIsTouch; // Is it a touch device
-	float FollowTime; // For how long it has been pressed
 };
 
 
diff --git a/Templates/TP_TopDown/Source/TP_TopDown/Variant_Strategy/StrategyGameMode.cpp b/Templates/TP_TopDown/Source/TP_TopDown/Variant_Strategy/StrategyGameMode.cpp
new file mode 100644
--- /dev/null
+++ b/Templates/TP_TopDown/Source/TP_TopDown/Variant_Strategy/StrategyGameMode.cpp
@@ -0,0 +1,5 @@
+// Copyright Epic Games, Inc. All Rights Reserved.
+
+
+#include "StrategyGameMode.h"
+
diff --git a/Templates/TP_TopDown/Source/TP_TopDown/Variant_Strategy/StrategyGameMode.h b/Templates/TP_TopDown/Source/TP_TopDown/Variant_Strategy/StrategyGameMode.h
new file mode 100644
--- /dev/null
+++ b/Templates/TP_TopDown/Source/TP_TopDown/Variant_Strategy/StrategyGameMode.h
@@ -0,0 +1,17 @@
+// Copyright Epic Games, Inc. All Rights Reserved.
+
+#pragma once
+
+#include "CoreMinimal.h"
+#include "GameFramework/GameModeBase.h"
+#include "StrategyGameMode.generated.h"
+
+/**
+ * 
+ */
+UCLASS(abstract)
+class AStrategyGameMode : public AGameModeBase
+{
+	GENERATED_BODY()
+	
+};
diff --git a/Templates/TP_TopDown/Source/TP_TopDown/Variant_Strategy/StrategyHUD.cpp b/Templates/TP_TopDown/Source/TP_TopDown/Variant_Strategy/StrategyHUD.cpp
new file mode 100644
--- /dev/null
+++ b/Templates/TP_TopDown/Source/TP_TopDown/Variant_Strategy/StrategyHUD.cpp
@@ -0,0 +1,77 @@
+// Copyright Epic Games, Inc. All Rights Reserved.
+
+
+#include "StrategyHUD.h"
+#include "StrategyUnit.h"
+#include "StrategyPlayerController.h"
+#include "StrategyUI.h"
+
+void AStrategyHUD::BeginPlay()
+{
+	Super::BeginPlay();
+
+	// spawn the UI widget
+	UIWidget = CreateWidget<UStrategyUI>(GetOwningPlayerController(), UIWidgetClass);
+	check(UIWidget);
+
+	// add the UI widget to the screen
+	UIWidget->AddToViewport(0);
+}
+
+void AStrategyHUD::DragSelectUpdate(FVector2D Start, FVector2D WidthAndHeight, FVector2D CurrentPosition, bool bDraw)
+{
+	// copy the selection box data
+	bDrawBox = bDraw;
+	BoxStart = Start;
+	BoxSize = WidthAndHeight;
+	BoxCurrentPosition = CurrentPosition;
+
+}
+
+void AStrategyHUD::DrawHUD()
+{
+	// draw all debug information, etc.
+	Super::DrawHUD();
+
+	// ensure we have a valid player controller
+	if (AStrategyPlayerController* PC = Cast<AStrategyPlayerController>(GetOwningPlayerController()))
+	{
+		// draw the selection box
+		if (bDrawBox)
+		{
+			DrawRect(SelectionBoxColor, BoxStart.X, BoxStart.Y, BoxSize.X, BoxSize.Y);
+
+			// get all the units in the selection box
+			TArray<AStrategyUnit*> BoxedUnits;
+			GetActorsInSelectionRectangle(BoxStart, BoxCurrentPosition, BoxedUnits, true);
+
+			// update the unit selection on the player controller
+			PC->DragSelectUnits(BoxedUnits);
+		}
+
+		// get the currently selected units
+		TArray<AStrategyUnit*> SelectedUnits = PC->GetSelectedUnits();
+
+		// update the selection count on the UI widget
+		UIWidget->SetSelectedUnitsCount(SelectedUnits.Num());
+
+		// process each selected unit
+		for (AStrategyUnit* CurrentUnit : SelectedUnits)
+		{
+			if (IsValid(CurrentUnit))
+			{
+				// project the unit's location to screen coordinates
+				FVector2D ScreenCoords;
+
+				if (PC->ProjectWorldLocationToScreen(CurrentUnit->GetActorLocation(), ScreenCoords, true))
+				{
+					// draw a selection string near the unit
+					const FString SelectionString = "Selected";
+					DrawText(SelectionString, FColor::White, ScreenCoords.X - 25.0f, ScreenCoords.Y + 25.0f, nullptr, 1.5f);
+				}
+			}
+			
+		}
+	}
+
+}
diff --git a/Templates/TP_TopDown/Source/TP_TopDown/Variant_Strategy/StrategyHUD.h b/Templates/TP_TopDown/Source/TP_TopDown/Variant_Strategy/StrategyHUD.h
new file mode 100644
--- /dev/null
+++ b/Templates/TP_TopDown/Source/TP_TopDown/Variant_Strategy/StrategyHUD.h
@@ -0,0 +1,57 @@
+// Copyright Epic Games, Inc. All Rights Reserved.
+
+#pragma once
+
+#include "CoreMinimal.h"
+#include "GameFramework/HUD.h"
+#include "StrategyHUD.generated.h"
+
+class UStrategyUI;
+
+/**
+ *  Simple strategy game HUD
+ *  Draws the selection box and unit selected overlays
+ */
+UCLASS(abstract)
+class AStrategyHUD : public AHUD
+{
+	GENERATED_BODY()
+	
+protected:
+
+	/** Pointer to the UI user widget */
+	TObjectPtr<UStrategyUI> UIWidget;
+
+	/** Type of UI Widget to spawn */
+	UPROPERTY(EditAnywhere, Category="UI")
+	TSubclassOf<UStrategyUI> UIWidgetClass;
+
+	/** If true, the HUD will draw the selection box */
+	bool bDrawBox = false;
+
+	/** Starting coords of the selection box */
+	FVector2D BoxStart;
+
+	/** Width and height of the selection box */
+	FVector2D BoxSize;
+
+	/** Current position of the selection box */
+	FVector2D BoxCurrentPosition;
+
+	/** Color of the selection box */
+	UPROPERTY(EditAnywhere, Category="UI")
+	FLinearColor SelectionBoxColor;
+
+public:
+
+	/** Initialization */
+	virtual void BeginPlay() override;
+
+	/** Updates the drag selection box */
+	void DragSelectUpdate(FVector2D Start, FVector2D WidthAndHeight, FVector2D CurrentPosition, bool bDraw);
+
+protected:
+
+	/** Draws the HUD */
+	virtual void DrawHUD() override;
+};
diff --git a/Templates/TP_TopDown/Source/TP_TopDown/Variant_Strategy/StrategyPawn.cpp b/Templates/TP_TopDown/Source/TP_TopDown/Variant_Strategy/StrategyPawn.cpp
new file mode 100644
--- /dev/null
+++ b/Templates/TP_TopDown/Source/TP_TopDown/Variant_Strategy/StrategyPawn.cpp
@@ -0,0 +1,39 @@
+// Copyright Epic Games, Inc. All Rights Reserved.
+
+
+#include "StrategyPawn.h"
+#include "Components/SceneComponent.h"
+#include "Camera/CameraComponent.h"
+#include "GameFramework/FloatingPawnMovement.h"
+
+AStrategyPawn::AStrategyPawn()
+{
+ 	PrimaryActorTick.bCanEverTick = true;
+
+	// create the root
+	RootComponent = CreateDefaultSubobject<USceneComponent>(TEXT("Root"));
+
+	// create the camera
+	Camera = CreateDefaultSubobject<UCameraComponent>(TEXT("Camera"));
+	Camera->SetupAttachment(RootComponent);
+
+	// create the movement component
+	FloatingPawnMovement = CreateDefaultSubobject<UFloatingPawnMovement>(TEXT("Floating Pawn Movement"));
+
+	// configure the camera
+	Camera->ProjectionMode = ECameraProjectionMode::Orthographic;
+	Camera->OrthoWidth = 1500.0f;
+	Camera->AutoPlaneShift = 1.0f;
+	Camera->bUpdateOrthoPlanes = false;
+
+	// configure the movement comp
+	FloatingPawnMovement->bConstrainToPlane = true;
+	FloatingPawnMovement->SetPlaneConstraintNormal(FVector::UpVector);
+	FloatingPawnMovement->SetPlaneConstraintOrigin(FVector::UpVector * 1500.0f);
+}
+
+void AStrategyPawn::SetZoomModifier(float Value)
+{
+	// set the ortho width on the camera
+	Camera->SetOrthoWidth(Value);
+}
diff --git a/Templates/TP_TopDown/Source/TP_TopDown/Variant_Strategy/StrategyPawn.h b/Templates/TP_TopDown/Source/TP_TopDown/Variant_Strategy/StrategyPawn.h
new file mode 100644
--- /dev/null
+++ b/Templates/TP_TopDown/Source/TP_TopDown/Variant_Strategy/StrategyPawn.h
@@ -0,0 +1,41 @@
+// Copyright Epic Games, Inc. All Rights Reserved.
+
+#pragma once
+
+#include "CoreMinimal.h"
+#include "GameFramework/Pawn.h"
+#include "StrategyPawn.generated.h"
+
+class UCameraComponent;
+class UFloatingPawnMovement;
+
+/**
+ *  Simple pawn that implements a top-down camera perspective for a strategy game.
+ *  Units are indirectly controlled by other means.
+ */
+UCLASS(abstract)
+class AStrategyPawn : public APawn
+{
+	GENERATED_BODY()
+
+	/** Camera */
+	UPROPERTY(VisibleAnywhere, BlueprintReadOnly, Category = "Components", meta = (AllowPrivateAccess = "true"))
+	UCameraComponent* Camera;
+
+	/** Movement Component */
+	UPROPERTY(VisibleAnywhere, BlueprintReadOnly, Category = "Camera", meta = (AllowPrivateAccess = "true"))
+	UFloatingPawnMovement* FloatingPawnMovement;
+
+public:
+
+	/** Constructor */
+	AStrategyPawn();
+
+public:
+
+	/** Sets the camera zoom modifier value */
+	void SetZoomModifier(float Value);
+
+	/** Returns the camera component */
+	UCameraComponent* GetCamera() const { return Camera; }
+};
diff --git a/Templates/TP_TopDown/Source/TP_TopDown/Variant_Strategy/StrategyPlayerController.cpp b/Templates/TP_TopDown/Source/TP_TopDown/Variant_Strategy/StrategyPlayerController.cpp
new file mode 100644
--- /dev/null
+++ b/Templates/TP_TopDown/Source/TP_TopDown/Variant_Strategy/StrategyPlayerController.cpp
@@ -0,0 +1,721 @@
+// Copyright Epic Games, Inc. All Rights Reserved.
+
+
+#include "StrategyPlayerController.h"
+#include "EnhancedInputSubsystems.h"
+#include "Engine/LocalPlayer.h"
+#include "EnhancedInputComponent.h"
+#include "InputMappingContext.h"
+#include "Camera/CameraComponent.h"
+#include "StrategyPawn.h"
+#include "Camera/CameraComponent.h"
+#include "InputActionValue.h"
+#include "StrategyHUD.h"
+#include "Engine/CollisionProfile.h"
+#include "Kismet/GameplayStatics.h"
+#include "StrategyUnit.h"
+#include "NavigationSystem.h"
+#include "Engine/OverlapResult.h"
+
+AStrategyPlayerController::AStrategyPlayerController()
+{
+	// mouse cursor should always be shown
+	bShowMouseCursor = true;
+}
+
+void AStrategyPlayerController::SetupInputComponent()
+{
+	Super::SetupInputComponent();
+
+	// add the input mapping context
+	if (UEnhancedInputLocalPlayerSubsystem* Subsystem = ULocalPlayer::GetSubsystem<UEnhancedInputLocalPlayerSubsystem>(GetLocalPlayer()))
+	{
+		// choose the context based on the input mode
+		UInputMappingContext* ChosenContext = nullptr;
+
+		switch (InputMode)
+		{
+		case SIM_Mouse:
+			ChosenContext = MouseMappingContext;
+			break;
+		case SIM_Touch:
+			ChosenContext = TouchMappingContext;
+			break;
+		}
+
+		Subsystem->AddMappingContext(ChosenContext, 0);
+	}
+
+	// bind the input mappings
+	if (UEnhancedInputComponent* EnhancedInputComponent = Cast<UEnhancedInputComponent>(InputComponent))
+	{
+		// Camera
+		EnhancedInputComponent->BindAction(MoveCameraAction, ETriggerEvent::Triggered, this, &AStrategyPlayerController::MoveCamera);
+		EnhancedInputComponent->BindAction(ZoomCameraAction, ETriggerEvent::Triggered, this, &AStrategyPlayerController::ZoomCamera);
+		EnhancedInputComponent->BindAction(ResetCameraAction, ETriggerEvent::Triggered, this, &AStrategyPlayerController::ResetCamera);
+
+		// Mouse Interaction
+		EnhancedInputComponent->BindAction(SelectHoldAction, ETriggerEvent::Started, this, &AStrategyPlayerController::SelectHoldStarted);
+		EnhancedInputComponent->BindAction(SelectHoldAction, ETriggerEvent::Triggered, this, &AStrategyPlayerController::SelectHoldTriggered);
+		EnhancedInputComponent->BindAction(SelectHoldAction, ETriggerEvent::Completed, this, &AStrategyPlayerController::SelectHoldCompleted);
+		EnhancedInputComponent->BindAction(SelectHoldAction, ETriggerEvent::Canceled, this, &AStrategyPlayerController::SelectHoldCompleted);
+
+		EnhancedInputComponent->BindAction(SelectClickAction, ETriggerEvent::Completed, this, &AStrategyPlayerController::SelectClick);
+
+		EnhancedInputComponent->BindAction(SelectionModifierAction, ETriggerEvent::Triggered, this, &AStrategyPlayerController::SelectionModifier);
+		EnhancedInputComponent->BindAction(SelectionModifierAction, ETriggerEvent::Completed, this, &AStrategyPlayerController::SelectionModifier);
+		EnhancedInputComponent->BindAction(SelectionModifierAction, ETriggerEvent::Canceled, this, &AStrategyPlayerController::SelectionModifier);
+
+		EnhancedInputComponent->BindAction(InteractHoldAction, ETriggerEvent::Started, this, &AStrategyPlayerController::InteractHoldStarted);
+		EnhancedInputComponent->BindAction(InteractHoldAction, ETriggerEvent::Triggered, this, &AStrategyPlayerController::InteractHoldTriggered);
+
+		EnhancedInputComponent->BindAction(InteractClickAction, ETriggerEvent::Started, this, &AStrategyPlayerController::InteractClickStarted);
+		EnhancedInputComponent->BindAction(InteractClickAction, ETriggerEvent::Completed, this, &AStrategyPlayerController::InteractClickCompleted);
+
+		// Touch Interaction
+		EnhancedInputComponent->BindAction(TouchPrimaryHoldAction, ETriggerEvent::Started, this, &AStrategyPlayerController::TouchPrimaryHoldStarted);
+		EnhancedInputComponent->BindAction(TouchPrimaryHoldAction, ETriggerEvent::Triggered, this, &AStrategyPlayerController::TouchPrimaryHoldTriggered);
+
+		EnhancedInputComponent->BindAction(TouchPrimaryTapAction, ETriggerEvent::Completed, this, &AStrategyPlayerController::TouchPrimaryTap);
+
+		EnhancedInputComponent->BindAction(TouchSecondaryAction, ETriggerEvent::Started, this, &AStrategyPlayerController::TouchSecondaryStarted);
+		EnhancedInputComponent->BindAction(TouchSecondaryAction, ETriggerEvent::Triggered, this, &AStrategyPlayerController::TouchSecondaryTriggered);
+		EnhancedInputComponent->BindAction(TouchSecondaryAction, ETriggerEvent::Completed, this, &AStrategyPlayerController::TouchSecondaryCompleted);
+		EnhancedInputComponent->BindAction(TouchSecondaryAction, ETriggerEvent::Canceled, this, &AStrategyPlayerController::TouchSecondaryCompleted);
+
+		EnhancedInputComponent->BindAction(TouchDoubleTapAction, ETriggerEvent::Triggered, this, &AStrategyPlayerController::TouchDoubleTap);
+	}
+}
+
+void AStrategyPlayerController::OnPossess(APawn* InPawn)
+{
+	Super::OnPossess(InPawn);
+
+	// ensure we have the right pawn type
+	ControlledPawn = Cast<AStrategyPawn>(InPawn);
+	check(ControlledPawn);
+
+	// set the zoom level from the pawn's camera
+	DefaultZoom = CameraZoom = ControlledPawn->GetCamera()->OrthoWidth;
+
+	// cast the HUD pointer
+	StrategyHUD = Cast<AStrategyHUD>(GetHUD());
+	check(StrategyHUD);
+}
+
+void AStrategyPlayerController::DragSelectUnits(const TArray<AStrategyUnit*>& Units)
+{
+	// do we have units in the list?
+	if (Units.Num() > 0)
+	{
+		// ensure any previous units are deselected
+		DoDeselectAllCommand();
+
+		// select each new unit
+		for (AStrategyUnit* CurrentUnit : Units)
+		{
+			// add the unit to the selection list
+			ControlledUnits.Add(CurrentUnit);
+
+			// select the unit
+			CurrentUnit->UnitSelected();
+		}
+	}
+}
+
+const TArray<AStrategyUnit*>& AStrategyPlayerController::GetSelectedUnits()
+{
+	return ControlledUnits;
+}
+
+void AStrategyPlayerController::MoveCamera(const FInputActionValue& Value)
+{
+	FVector2D InputVector = Value.Get<FVector2D>();
+
+	// get the forward input component vector
+	FRotator ForwardRot = GetControlRotation();
+	ForwardRot.Pitch = 0.0f;
+
+	// get the right input component vector
+	FRotator RightRot = GetControlRotation();
+	ForwardRot.Pitch = 0.0f;
+	ForwardRot.Roll = 0.0f;
+
+	// add the forward input
+	ControlledPawn->AddMovementInput(ForwardRot.RotateVector(FVector::ForwardVector), InputVector.X + InputVector.Y);
+
+	// add the right input
+	ControlledPawn->AddMovementInput(RightRot.RotateVector(FVector::RightVector), InputVector.X - InputVector.Y);
+
+}
+
+void AStrategyPlayerController::ZoomCamera(const FInputActionValue& Value)
+{
+	// scale the input and subtract from the current zoom level
+	float ZoomLevel = CameraZoom - (Value.Get<float>() * ZoomScaling);
+
+	// clamp to min/max zoom levels
+	CameraZoom = FMath::Clamp(ZoomLevel, MinZoomLevel, MaxZoomLevel);
+
+	// update the pawn's camera
+	ControlledPawn->SetZoomModifier(CameraZoom);
+
+}
+
+void AStrategyPlayerController::ResetCamera(const FInputActionValue& Value)
+{
+	// reset zoom level to its initial value
+	CameraZoom = DefaultZoom;
+
+	// update the pawn's camera
+	ControlledPawn->SetZoomModifier(DefaultZoom);
+
+}
+
+void AStrategyPlayerController::SelectHoldStarted(const FInputActionValue& Value)
+{
+	// save the selection start position
+	StartingSelectionPosition = GetMouseLocation();
+
+}
+
+void AStrategyPlayerController::SelectHoldTriggered(const FInputActionValue& Value)
+{
+
+	// get the current mouse position
+	FVector2D SelectionPosition = GetMouseLocation();
+
+	// calculate the size of the selection box
+	FVector2D SelectionSize = SelectionPosition - StartingSelectionPosition;
+
+	// update the selection box on the HUD
+	StrategyHUD->DragSelectUpdate(StartingSelectionPosition, SelectionSize, SelectionPosition, true);
+}
+
+void AStrategyPlayerController::SelectHoldCompleted(const FInputActionValue& Value)
+{
+
+	// reset the drag box on the HUD
+	StrategyHUD->DragSelectUpdate(FVector2D::ZeroVector, FVector2D::ZeroVector, FVector2D::ZeroVector, false);
+}
+
+void AStrategyPlayerController::SelectClick(const FInputActionValue& Value)
+{
+
+	if (GetLocationUnderCursor(CachedSelection))
+	{
+		DoSelectionCommand();
+	}
+}
+
+void AStrategyPlayerController::SelectionModifier(const FInputActionValue& Value)
+{
+
+	// update the selection modifier flag
+	bSelectionModifier = Value.Get<bool>();
+}
+
+void AStrategyPlayerController::InteractHoldStarted(const FInputActionValue& Value)
+{
+
+	// save the starting interaction position
+	StartingInteractionPosition = GetMouseLocation();
+}
+
+void AStrategyPlayerController::InteractHoldTriggered(const FInputActionValue& Value)
+{
+
+	// do a drag scroll 
+	DoDragScrollCommand();
+}
+
+void AStrategyPlayerController::InteractClickStarted(const FInputActionValue& Value)
+{
+
+	// reset the interaction flag
+	ResetInteraction();
+}
+
+void AStrategyPlayerController::InteractClickCompleted(const FInputActionValue& Value)
+{
+
+	// do we have any units in the control list and a valid interaction location under the cursor?
+	if (ControlledUnits.Num() > 0 && GetLocationUnderCursor(CachedInteraction))
+	{
+		// is double tap select all active?
+		if (bDoubleTapActive)
+		{
+			// release double tap select all
+			bDoubleTapActive = false;
+
+		}
+		else {
+
+			// move the selected units to the target location
+			DoMoveUnitsCommand();
+
+		}
+	}
+}
+
+void AStrategyPlayerController::TouchPrimaryHoldStarted(const FInputActionValue& Value)
+{
+
+	// save the starting interaction position
+	StartingInteractionPosition = Value.Get<FVector2D>();
+}
+
+void AStrategyPlayerController::TouchPrimaryHoldTriggered(const FInputActionValue& Value)
+{
+
+	// are we in selection modifier mode, and have a large enough selection box?
+	if (bSelectionModifier && StartingSecondFingerPosition.Equals(CurrentSecondFingerPosition, 10.0f))
+	{
+		// update the interaction position
+		CurrentInteractionPosition = Value.Get<FVector2D>();
+
+		// update the selection box on the HUD
+		StrategyHUD->DragSelectUpdate(StartingInteractionPosition, CurrentInteractionPosition - StartingSecondFingerPosition, CurrentInteractionPosition, true);
+
+	} else {
+
+		// do a drag scroll instead
+		DoDragScrollCommand();
+
+	}
+}
+
+void AStrategyPlayerController::TouchPrimaryTap(const FInputActionValue& Value)
+{
+
+	// project the touch location and cache the selection point
+	CachedSelection = ProjectTouchPointToWorldSpace();
+
+	// do a selection action with the cached location
+	DoSelectionCommand();
+}
+
+void AStrategyPlayerController::TouchSecondaryStarted(const FInputActionValue& Value)
+{
+
+	// raise the selection modifier flag
+	bSelectionModifier = true;
+
+	// save the starting position for the second finger
+	StartingSecondFingerPosition = Value.Get<FVector2D>();
+}
+
+void AStrategyPlayerController::TouchSecondaryTriggered(const FInputActionValue& Value)
+{
+
+	// update the current position for the second finger
+	CurrentSecondFingerPosition = Value.Get<FVector2D>();
+}
+
+void AStrategyPlayerController::TouchSecondaryCompleted(const FInputActionValue& Value)
+{
+
+	// lower the selection modifier flag
+	bSelectionModifier = false;
+}
+
+void AStrategyPlayerController::TouchDoubleTap(const FInputActionValue& Value)
+{
+
+	// raise the double tap flag
+	bDoubleTapActive = true;
+
+	// is the selection modifier mode active?
+	if (bSelectionModifier)
+	{
+		// deselect all units
+		DoDeselectAllCommand();
+
+	} else {
+
+		// select all units on screen
+		DoSelectAllOnScreenCommand();
+	}
+}
+
+void AStrategyPlayerController::DoSelectionCommand()
+{
+
+	// do a sphere sweep to look for actors to select
+	FHitResult OutHit;
+
+	const FVector Start = CachedSelection;
+	const FVector End = Start + FVector::UpVector * 350.0f;
+
+	FCollisionShape InteractionSphere;
+	InteractionSphere.SetSphere(InteractionRadius);
+
+	FCollisionObjectQueryParams ObjectParams;
+	ObjectParams.AddObjectTypesToQuery(ECC_Pawn);
+
+	FCollisionQueryParams QueryParams;
+	QueryParams.AddIgnoredActor(this);
+	QueryParams.AddIgnoredActor(GetPawn());
+	QueryParams.bTraceComplex = true;
+
+	GetWorld()->SweepSingleByObjectType(OutHit, Start, End, FQuat::Identity, ObjectParams, InteractionSphere, QueryParams);
+
+	// if we're using the mouse and are not holding the selection modifier key, deselect any units first
+	if (InputMode == SIM_Mouse && !bSelectionModifier)
+	{
+
+		DoDeselectAllCommand();
+	}
+
+	// did we hit a unit?
+	if (OutHit.bBlockingHit)
+	{
+
+		// update the target unit
+		TargetUnit = Cast<AStrategyUnit>(OutHit.GetActor());
+
+		if (TargetUnit)
+		{
+
+			// is the unit already in the controlled list?
+			if (ControlledUnits.Contains(TargetUnit))
+			{
+
+				// remove the units from the controlled list
+				ControlledUnits.Remove(TargetUnit);
+
+				// tell the unit it's been deselected
+				TargetUnit->UnitDeselected();
+
+			}
+			else {
+
+				// add the unit to the controlled list
+				ControlledUnits.Add(TargetUnit);
+
+				// tell the unit it's been selected
+				TargetUnit->UnitSelected();
+
+			}
+		}
+
+	} else {
+
+		// are we using touch input?
+		if (InputMode == SIM_Touch)
+		{
+			// is the double tap select all flag set?
+			if (bDoubleTapActive)
+			{
+
+				// release double tap select all
+				bDoubleTapActive = false;
+
+			} else {
+
+				// move all selected units to the target location
+				DoMoveUnitsCommand();
+
+			}
+		}
+
+	}
+}
+
+void AStrategyPlayerController::DoSelectAllOnScreenCommand()
+{
+
+	// find all NPCs currently on screen
+	TArray<AActor*> FoundActors;
+	UGameplayStatics::GetAllActorsOfClass(GetWorld(), AStrategyUnit::StaticClass(), FoundActors);
+
+	// process each actor found
+	for (AActor* CurrentActor : FoundActors)
+	{
+		// cast back to our unit class
+		if (AStrategyUnit* CurrentUnit = Cast<AStrategyUnit>(CurrentActor))
+		{
+			// has the actor been recently rendered?
+			if (CurrentActor->WasRecentlyRendered(0.2f))
+			{
+
+				// is the actor not on our controlled units list?
+				if (!ControlledUnits.Contains(CurrentUnit))
+				{
+					// add it to the controlled units list
+					ControlledUnits.Add(CurrentUnit);
+
+					// notify it of selection
+					CurrentUnit->UnitSelected();
+				}
+			}
+		}		
+	}
+
+}
+
+void AStrategyPlayerController::DoDeselectAllCommand()
+{
+
+	// tell each controlled unit it's been deselected
+	for (AStrategyUnit* CurrentUnit : ControlledUnits)
+	{
+		// ensure the unit hasn't been destroyed
+		if (IsValid(CurrentUnit))
+		{
+
+			CurrentUnit->UnitDeselected();
+		}
+	}
+
+	// clear the controlled units list
+	ControlledUnits.Empty();
+}
+
+void AStrategyPlayerController::DoDragScrollCommand()
+{
+
+	// choose the cursor position based on the input mode
+	FVector2D WorkingPosition;
+	
+	if (InputMode == EStrategyInputMode::SIM_Mouse)
+	{
+
+		// read the mouse position
+		bool bResult = GetMousePosition(WorkingPosition.X, WorkingPosition.Y);
+
+	} else {
+
+		// read the touch 1 position
+		bool bPressed;
+		GetInputTouchState(ETouchIndex::Touch1, WorkingPosition.X, WorkingPosition.Y, bPressed);
+
+	}
+
+	// find the difference between the starting interaction position and current coords
+	const FVector2D InteractionDelta = StartingInteractionPosition - WorkingPosition;
+
+	const FRotator CameraRot(0.0f, -45.0f, 0.0f);
+
+	// rotate and scale the interaction delta
+	const FVector ScrollDelta = CameraRot.RotateVector(FVector(InteractionDelta.X, InteractionDelta.Y, 0.0f)) * DragMultiplier;
+
+	// apply the world offset to the controlled pawn
+	ControlledPawn->AddActorWorldOffset(ScrollDelta);
+}
+
+void AStrategyPlayerController::DoMoveUnitsCommand()
+{
+
+	// set the movement goal
+	FVector CurrentMoveGoal;
+
+	if (InputMode == EStrategyInputMode::SIM_Mouse)
+	{
+
+		// set the cached interaction point as our move goal
+		CurrentMoveGoal = CachedInteraction;
+
+	}
+	else {
+
+		// set the cached selection as our move goal
+		CurrentMoveGoal = CachedSelection;
+
+	}
+
+	// get the closest selected unit to the move goal. This will be our lead unit
+	AStrategyUnit* Closest = GetClosestSelectedUnitToLocation(CurrentMoveGoal);
+
+	// this will be set to true if any of the move requests fail
+	bool bInteractionFailed = false;
+
+	// process each unit in the controlled list
+	for (AStrategyUnit* CurrentUnit : ControlledUnits)
+	{
+		if (IsValid(CurrentUnit))
+		{
+
+			// stop the unit
+			CurrentUnit->StopMoving();
+
+			// move the lead unit to the goal, all other units to random navigable points around it
+			FVector MoveGoal = CurrentMoveGoal;
+
+			if (CurrentUnit != Closest)
+			{
+
+				UNavigationSystemV1::K2_GetRandomLocationInNavigableRadius(GetWorld(), CurrentMoveGoal, MoveGoal, InteractionRadius * 0.66f);
+			}
+
+			// subscribe to the unit's move completed delegate
+			CurrentUnit->OnMoveCompleted.AddDynamic(this, &AStrategyPlayerController::OnMoveCompleted);
+
+			// set up movement to the goal location
+			if (!CurrentUnit->MoveToLocation(MoveGoal, InteractionRadius * 0.66f))
+			{
+				// the move request failed, so flag it
+				bInteractionFailed = true;
+			}
+		}
+
+	}
+
+	// play the cursor feedback depending on whether our move succeeded or not
+	BP_CursorFeedback(CachedInteraction, !bInteractionFailed);
+
+}
+
+void AStrategyPlayerController::OnMoveCompleted(AStrategyUnit* MovedUnit)
+{
+	// is the unit valid?
+	if (IsValid(MovedUnit))
+	{
+		// unsubscribe from the delegate
+		MovedUnit->OnMoveCompleted.RemoveDynamic(this, &AStrategyPlayerController::OnMoveCompleted);
+		
+		// skip if interactions are locked
+		if (!bAllowInteraction)
+		{
+			return;
+		}
+
+		// disallow additional interactions until we reset
+		bAllowInteraction = false;
+
+		// is the unit close enough to the cached interaction location?
+		if(FVector::Dist2D(CachedInteraction, MovedUnit->GetActorLocation()) < InteractionRadius)
+		{
+
+			// do an overlap test to find nearby interactive objects
+			TArray<FOverlapResult> OutOverlaps;
+
+			FCollisionShape CollisionSphere;
+			CollisionSphere.SetSphere(InteractionRadius);
+
+			FCollisionObjectQueryParams ObjectParams;
+			ObjectParams.AddObjectTypesToQuery(ECC_WorldDynamic);
+			
+			FCollisionQueryParams QueryParams;
+
+			QueryParams.AddIgnoredActor(MovedUnit);
+
+			for(const AStrategyUnit* CurSelected : ControlledUnits)
+			{
+				QueryParams.AddIgnoredActor(CurSelected);
+			}
+
+			if (GetWorld()->OverlapMultiByObjectType(OutOverlaps, CachedInteraction, FQuat::Identity, ObjectParams, CollisionSphere, QueryParams))
+			{
+				for (const FOverlapResult& CurrentOverlap : OutOverlaps)
+				{
+					if (AStrategyUnit* CurrentUnit = Cast<AStrategyUnit>(CurrentOverlap.GetActor()))
+					{
+						CurrentUnit->Interact(MovedUnit);
+					}
+				}
+			}
+		}
+	}
+}
+
+AStrategyUnit* AStrategyPlayerController::GetClosestSelectedUnitToLocation(FVector TargetLocation)
+{
+	// closest unit and distance
+	AStrategyUnit* OutUnit = nullptr;
+	float Closest = 0.0f;
+
+	// process each unit on the list
+	for (AStrategyUnit* CurrentUnit : ControlledUnits)
+	{
+		if (CurrentUnit != nullptr)
+		{
+			// have we selected a unit already?
+			if (OutUnit != nullptr)
+			{
+				// calculate the squared distance to the target location
+				float Dist = FVector::DistSquared2D(TargetLocation, CurrentUnit->GetActorLocation());
+
+				// is this unit closer?
+				if (Dist < Closest)
+				{
+					// update the closest unit and distance
+					OutUnit = CurrentUnit;
+					Closest = Dist;
+				}
+
+			} else {
+
+				// no previously selected unit, so use this one
+				OutUnit = CurrentUnit;
+
+				// initialize the closest distance
+				Closest = FVector::DistSquared2D(TargetLocation, CurrentUnit->GetActorLocation());
+			}
+		}
+		
+	}
+
+	// return the selected unit
+	return OutUnit;
+}
+
+FVector2D AStrategyPlayerController::GetMouseLocation()
+{
+	// attempt to get the mouse position from this PC
+	float MouseX, MouseY;
+
+	if (GetMousePosition(MouseX, MouseY))
+	{
+		return FVector2D(MouseX, MouseY);
+	}
+
+	// return an invalid vector
+	return FVector2D::ZeroVector;
+}
+
+bool AStrategyPlayerController::GetLocationUnderCursor(FVector& Location)
+{
+	// trace the visibility channel at the cursor location
+	FHitResult OutHit;
+
+	GetHitResultUnderCursorByChannel(SelectionTraceChannel, true, OutHit);
+
+	// if there was a blocking hit, return the hit location
+	if (OutHit.bBlockingHit)
+	{
+		Location = OutHit.Location;
+		return true;
+	}
+
+	return OutHit.bBlockingHit;
+}
+
+FVector AStrategyPlayerController::ProjectTouchPointToWorldSpace()
+{
+	// get the touch coordinates for the first finger
+	float TouchX, TouchY = 0.0f;
+	bool bPressed = false;
+
+	GetInputTouchState(ETouchIndex::Touch1, TouchX, TouchY, bPressed);
+
+	FVector WorldLocation = FVector::ZeroVector;
+	FVector WorldDirection = FVector::ZeroVector;
+
+	// deproject the coords into world space
+	if (DeprojectScreenPositionToWorld(TouchX, TouchY, WorldLocation, WorldDirection))
+	{
+		// intersect with a horizontal plane and return the resulting point
+		const FPlane IntersectPlane(FVector::ZeroVector, FVector::UpVector);
+
+		return FMath::LinePlaneIntersection(WorldLocation, WorldLocation + (WorldDirection * 100000.0f), IntersectPlane);
+	}
+
+	// failed to deproject, return a zero vector
+	return FVector::ZeroVector;
+}
+
+void AStrategyPlayerController::ResetInteraction()
+{
+	bAllowInteraction = true;
+}
diff --git a/Templates/TP_TopDown/Source/TP_TopDown/Variant_Strategy/StrategyPlayerController.h b/Templates/TP_TopDown/Source/TP_TopDown/Variant_Strategy/StrategyPlayerController.h
new file mode 100644
--- /dev/null
+++ b/Templates/TP_TopDown/Source/TP_TopDown/Variant_Strategy/StrategyPlayerController.h
@@ -0,0 +1,285 @@
+// Copyright Epic Games, Inc. All Rights Reserved.
+
+#pragma once
+
+#include "CoreMinimal.h"
+#include "GameFramework/PlayerController.h"
+#include "StrategyPlayerController.generated.h"
+
+class AStrategyPawn;
+class UInputMappingContext;
+class UNiagaraSystem;
+struct FInputActionValue;
+class AStrategyHUD;
+class AStrategyNPC;
+class UInputAction;
+
+UENUM(BlueprintType)
+enum EStrategyInputMode : uint8
+{
+	SIM_Mouse	UMETA(DisplayName = "Mouse"),
+	SIM_Touch	UMETA(DisplayName = "Touch")
+};
+
+/**
+ * 
+ */
+UCLASS(abstract)
+class AStrategyPlayerController : public APlayerController
+{
+	GENERATED_BODY()
+
+protected:
+
+	/** Strategy Pawn associated with this controller */
+	TObjectPtr<AStrategyPawn> ControlledPawn;
+
+	/** Strategy HUD associated with this controller */
+	TObjectPtr<AStrategyHUD> StrategyHUD;
+
+	/** Determines the chosen input type */
+	UPROPERTY(EditAnywhere, Category = "Input")
+	TEnumAsByte<EStrategyInputMode> InputMode = SIM_Mouse;
+
+	/** Input mapping context to use with mouse input */
+	UPROPERTY(EditAnywhere, Category = "Input")
+	UInputMappingContext* MouseMappingContext;
+
+	/** Input mapping context to use with touch input */
+	UPROPERTY(EditAnywhere, Category = "Input")
+	UInputMappingContext* TouchMappingContext;
+
+	/** If true, the player is adding or removing units to the selected units list */
+	bool bSelectionModifier = false;
+
+	/** If true, double-tap touch select all mode is active */
+	bool bDoubleTapActive = false;
+
+	/** If true, allow the player to interact with game objects */
+	bool bAllowInteraction = true;
+
+	/** Input Action for moving the camera */
+	UPROPERTY(EditAnywhere, Category = "Input")
+	UInputAction* MoveCameraAction;
+
+	/** Input Action for zooming the camera */
+	UPROPERTY(EditAnywhere, Category = "Input")
+	UInputAction* ZoomCameraAction;
+
+	/** Input Action for resetting the camera to its default position */
+	UPROPERTY(EditAnywhere, Category = "Input")
+	UInputAction* ResetCameraAction;
+
+	/** Input Action for select and click */
+	UPROPERTY(EditAnywhere, Category = "Input")
+	UInputAction* SelectClickAction;
+
+	/** Input Action for select press and hold */
+	UPROPERTY(EditAnywhere, Category = "Input")
+	UInputAction* SelectHoldAction;
+
+	/** Input Action for click interaction */
+	UPROPERTY(EditAnywhere, Category = "Input")
+	UInputAction* InteractClickAction;
+
+	/** Input Action for interaction press and hold */
+	UPROPERTY(EditAnywhere, Category = "Input")
+	UInputAction* InteractHoldAction;
+
+	/** Input Action for modifying selection mode */
+	UPROPERTY(EditAnywhere, Category = "Input")
+	UInputAction* SelectionModifierAction;
+
+	/** Input Action for primary touch tap */
+	UPROPERTY(EditAnywhere, Category = "Input")
+	UInputAction* TouchPrimaryTapAction;
+
+	/** Input Action for primary touch hold */
+	UPROPERTY(EditAnywhere, Category = "Input")
+	UInputAction* TouchPrimaryHoldAction;
+
+	/** Input Action for secondary touch */
+	UPROPERTY(EditAnywhere, Category = "Input")
+	UInputAction* TouchSecondaryAction;
+
+	/** Input Action for touch double tap */
+	UPROPERTY(EditAnywhere, Category = "Input")
+	UInputAction* TouchDoubleTapAction;
+
+	/** Max distance to look for nearby units when doing a click or touch interaction */
+	UPROPERTY(EditAnywhere, Category = "Input", meta = (ClampMin = 0, ClampMax = 10000, Units = "cm"))
+	float InteractionRadius = 250.0f;
+
+	/** Max distance between the starting and current position of the second touch finger to be considered a box selection */
+	UPROPERTY(EditAnywhere, Category = "Input", meta = (ClampMin = 0, ClampMax = 10000))
+	float MinSecondFingerDistanceForBoxSelect = 10.0f;
+
+	/** Saves the world location of the last initiated interaction */
+	FVector CachedInteraction;
+
+	/** Saves the world location of the last initiated unit selection */
+	FVector CachedSelection;
+
+	/** Saves the world location where the player started a press and hold interaction */
+	FVector2D StartingInteractionPosition;
+
+	/** Saves the current world location of the player's cursor in press and hold interaction */
+	FVector2D CurrentInteractionPosition;
+
+	/** Saves the starting world location of a player's cursor in a press and hold selection box */
+	FVector2D StartingSelectionPosition;
+
+	/** Saves the starting location of a two-finger touch interaction (pinch) */
+	FVector2D StartingSecondFingerPosition;
+
+	/** Saves the current location of a two-finger touch interaction (pinch) */
+	FVector2D CurrentSecondFingerPosition;
+
+	/** Current camera zoom level */
+	float CameraZoom;
+
+	/** Default camera zoom level */
+	float DefaultZoom;
+
+	/** Minimum allowed camera zoom level */
+	UPROPERTY(EditAnywhere, Category = "Camera", meta = (ClampMin = 0, ClampMax = 10000))
+	float MinZoomLevel = 1000.0f;
+
+	/** Maximum allowed camera zoom level */
+	UPROPERTY(EditAnywhere, Category = "Camera", meta = (ClampMin = 0, ClampMax = 10000))
+	float MaxZoomLevel = 2500.0f;
+
+	/** Scales zoom inputs by this value */
+	UPROPERTY(EditAnywhere, Category = "Camera", meta = (ClampMin = 0, ClampMax = 1000))
+	float ZoomScaling = 100.0f;
+
+	/** Affects how fast the camera moves while dragging with the mouse */
+	UPROPERTY(EditAnywhere, Category = "Camera", meta = (ClampMin = 0, ClampMax = 10000))
+	float DragMultiplier = 0.1f;
+
+	/** Trace channel to use for selection trace checks */
+	UPROPERTY(EditAnywhere, Category = "Selection")
+	TEnumAsByte<ETraceTypeQuery> SelectionTraceChannel;
+
+	/** Currently selected unit */
+	AStrategyUnit* TargetUnit = nullptr;
+
+	/** Currently selected unit list */
+	TArray<AStrategyUnit*> ControlledUnits;
+
+public:
+
+	/** Constructor */
+	AStrategyPlayerController();
+
+	/** Initialize input bindings */
+	virtual void SetupInputComponent() override;
+
+	/** Pawn initialization */
+	virtual void OnPossess(APawn* InPawn);
+
+public:
+
+	/** Updates selected units from the HUD's drag select box */
+	void DragSelectUnits(const TArray<AStrategyUnit*>& Units);
+
+	/** Passes the list of selected units */
+	const TArray<AStrategyUnit*>& GetSelectedUnits();
+
+protected:
+
+	/** Moves the camera by the given input */
+	void MoveCamera(const FInputActionValue& Value);
+
+	/** Changes the camera zoom level by the given input */
+	void ZoomCamera(const FInputActionValue& Value);
+
+	/** Resets the camera to its initial value */
+	void ResetCamera(const FInputActionValue& Value);
+
+	/** Start a select and hold input */
+	void SelectHoldStarted(const FInputActionValue& Value);
+	
+	/** Select and hold input triggered */
+	void SelectHoldTriggered(const FInputActionValue& Value);
+
+	/** Select and hold input completed */
+	void SelectHoldCompleted(const FInputActionValue& Value);
+
+	/** Select click action */
+	void SelectClick(const FInputActionValue& Value);
+
+	/** Presses or releases the selection modifier key */
+	void SelectionModifier(const FInputActionValue& Value);
+
+	/** Starts an interaction hold input */
+	void InteractHoldStarted(const FInputActionValue& Value);
+
+	/** Interaction hold input triggered */
+	void InteractHoldTriggered(const FInputActionValue& Value);
+
+	/** Interaction click input started */
+	void InteractClickStarted(const FInputActionValue& Value);
+
+	/** Interaction click input completed */
+	void InteractClickCompleted(const FInputActionValue& Value);
+
+	/** Touch primary finger hold started */
+	void TouchPrimaryHoldStarted(const FInputActionValue& Value);
+
+	/** Touch primary finger hold triggered */
+	void TouchPrimaryHoldTriggered(const FInputActionValue& Value);
+
+	/** Touch primary finger tap completed */
+	void TouchPrimaryTap(const FInputActionValue& Value);
+
+	/** Touch secondary finger started */
+	void TouchSecondaryStarted(const FInputActionValue& Value);
+
+	/** Touch secondary finger triggered */
+	void TouchSecondaryTriggered(const FInputActionValue& Value);
+
+	/** Touch secondary finger completed */
+	void TouchSecondaryCompleted(const FInputActionValue& Value);
+
+	/** Touch primary finger double tap triggered */
+	void TouchDoubleTap(const FInputActionValue& Value);
+
+	/** Attempt to select or deselect units at the cached location */
+	void DoSelectionCommand();
+
+	/** Select all units currently on screen */
+	void DoSelectAllOnScreenCommand();
+
+	/** Deselect all controlled units */
+	void DoDeselectAllCommand();
+
+	/** Drag scroll the camera */
+	void DoDragScrollCommand();
+
+	/** Move all selected units */
+	void DoMoveUnitsCommand();
+
+	/** Called when a unit move is completed */
+	UFUNCTION()
+	void OnMoveCompleted(AStrategyUnit* MovedUnit);
+
+	/** Sorts all controlled units based on their distance to the provided world location */
+	AStrategyUnit* GetClosestSelectedUnitToLocation(FVector TargetLocation);
+
+	/** Calculates and returns the current mouse location */
+	FVector2D GetMouseLocation();
+
+	/** Attempts to get the world location under the cursor, returns true if successful */
+	bool GetLocationUnderCursor(FVector& Location);
+
+	/** Projects the current touch location into world space */
+	FVector ProjectTouchPointToWorldSpace();
+
+	/** Spawns the positive cursor effect */
+	UFUNCTION(BlueprintImplementableEvent, Category="Cursor", meta=(DisplayName="Cursor Feedback"))
+	void BP_CursorFeedback(FVector Location, bool bPositive);
+
+	/** Resets the interaction flag */
+	void ResetInteraction();
+};
diff --git a/Templates/TP_TopDown/Source/TP_TopDown/Variant_Strategy/StrategyUI.cpp b/Templates/TP_TopDown/Source/TP_TopDown/Variant_Strategy/StrategyUI.cpp
new file mode 100644
--- /dev/null
+++ b/Templates/TP_TopDown/Source/TP_TopDown/Variant_Strategy/StrategyUI.cpp
@@ -0,0 +1,19 @@
+// Copyright Epic Games, Inc. All Rights Reserved.
+
+
+#include "StrategyUI.h"
+
+void UStrategyUI::SetSelectedUnitsCount(int32 Count)
+{
+	// is this a different count?
+	bool bChanged = SelectedUnitCount != Count;
+
+	// update the counter
+	SelectedUnitCount = Count;
+
+	// if the count changed, call the BP handler
+	if (bChanged)
+	{
+		BP_UpdateUnitsCount();
+	}
+}
diff --git a/Templates/TP_TopDown/Source/TP_TopDown/Variant_Strategy/StrategyUI.h b/Templates/TP_TopDown/Source/TP_TopDown/Variant_Strategy/StrategyUI.h
new file mode 100644
--- /dev/null
+++ b/Templates/TP_TopDown/Source/TP_TopDown/Variant_Strategy/StrategyUI.h
@@ -0,0 +1,37 @@
+// Copyright Epic Games, Inc. All Rights Reserved.
+
+#pragma once
+
+#include "CoreMinimal.h"
+#include "Blueprint/UserWidget.h"
+#include "StrategyUI.generated.h"
+
+/**
+ *  Simple UI widget for the strategy game
+ *	Keeps track of the number of units currently selected
+ */
+UCLASS(abstract)
+class UStrategyUI : public UUserWidget
+{
+	GENERATED_BODY()
+	
+protected:
+
+	/** Number of units currently selected */
+	int32 SelectedUnitCount = 0;
+
+public:
+
+	/** Sets the number of units selected */
+	void SetSelectedUnitsCount(int32 Count);
+
+	/** Blueprint handler to update unit count sub-widgets */
+	UFUNCTION(BlueprintImplementableEvent, Category="UI", meta=(DisplayName="Update Units Count"))
+	void BP_UpdateUnitsCount();
+
+protected:
+
+	/** Returns the number of units selected */
+	UFUNCTION(BlueprintPure, Category="UI")
+	int32 GetSelectedUnitsCount() { return SelectedUnitCount; }
+};
diff --git a/Templates/TP_TopDown/Source/TP_TopDown/Variant_Strategy/StrategyUnit.cpp b/Templates/TP_TopDown/Source/TP_TopDown/Variant_Strategy/StrategyUnit.cpp
new file mode 100644
--- /dev/null
+++ b/Templates/TP_TopDown/Source/TP_TopDown/Variant_Strategy/StrategyUnit.cpp
@@ -0,0 +1,146 @@
+// Copyright Epic Games, Inc. All Rights Reserved.
+
+
+#include "StrategyUnit.h"
+#include "AIController.h"
+#include "GameFramework/CharacterMovementComponent.h"
+#include "Kismet/KismetMathLibrary.h"
+#include "Components/SphereComponent.h"
+#include "Navigation/PathFollowingComponent.h"
+
+AStrategyUnit::AStrategyUnit()
+{
+	PrimaryActorTick.bCanEverTick = true;
+
+	// ensure this unit has a valid AI controller to handle move requests
+	AutoPossessAI = EAutoPossessAI::PlacedInWorldOrSpawned;
+
+	// create the interaction range sphere
+	InteractionRange = CreateDefaultSubobject<USphereComponent>(TEXT("Interaction Range"));
+	InteractionRange->SetupAttachment(RootComponent);
+
+	InteractionRange->SetSphereRadius(100.0f);
+	InteractionRange->SetCollisionProfileName(FName("OverlapAllDynamic"));
+
+	// configure movement
+	GetCharacterMovement()->GravityScale = 1.5f;
+	GetCharacterMovement()->MaxAcceleration = 1000.0f;
+	GetCharacterMovement()->BrakingFrictionFactor = 1.0f;
+	GetCharacterMovement()->BrakingDecelerationWalking = 1000.0f;
+	GetCharacterMovement()->PerchRadiusThreshold = 20.0f;
+	GetCharacterMovement()->bUseFlatBaseForFloorChecks = true;
+	GetCharacterMovement()->RotationRate = FRotator(0.0f, 640.0f, 0.0f);
+	GetCharacterMovement()->bOrientRotationToMovement = true;
+	GetCharacterMovement()->AvoidanceConsiderationRadius = 150.0f;
+	GetCharacterMovement()->AvoidanceWeight = 1.0f;
+	GetCharacterMovement()->bConstrainToPlane = true;
+	GetCharacterMovement()->bSnapToPlaneAtStart = true;
+	GetCharacterMovement()->SetFixedBrakingDistance(200.0f);
+	GetCharacterMovement()->SetFixedBrakingDistance(true);
+}
+
+void AStrategyUnit::NotifyControllerChanged()
+{
+	// validate and save a copy of the AI controller reference
+	AIController = Cast<AAIController>(Controller);
+	
+	if (AIController)
+	{
+		// subscribe to the move finished handler on the path following component
+		UPathFollowingComponent* PFComp = AIController->GetPathFollowingComponent();
+		if (PFComp)
+		{
+			PFComp->OnRequestFinished.AddUObject(this, &AStrategyUnit::OnMoveFinished);
+		}
+	}
+}
+
+void AStrategyUnit::StopMoving()
+{
+	// use the character movement component to stop movement
+	GetCharacterMovement()->StopMovementImmediately();
+}
+
+void AStrategyUnit::UnitSelected()
+{
+	// pass control to BP
+	BP_UnitSelected();
+}
+
+void AStrategyUnit::UnitDeselected()
+{
+	// pass control to BP
+	BP_UnitDeselected();
+}
+
+void AStrategyUnit::Interact(AStrategyUnit* Interactor)
+{
+	// ensure the interactor is valid
+	if (IsValid(Interactor))
+	{
+		// rotate towards the actor we're interacting with
+		SetActorRotation(UKismetMathLibrary::FindLookAtRotation(GetActorLocation(), Interactor->GetActorLocation()));
+
+		// signal the interactor to play its interaction behavior
+		Interactor->BP_InteractionBehavior(this);
+
+		// play our own interaction behavior
+		BP_InteractionBehavior(Interactor);
+	}
+	
+}
+
+bool AStrategyUnit::MoveToLocation(const FVector& Location, float AcceptanceRadius)
+{
+	// ensure we have a valid AI Controller
+	if (AIController)
+	{
+		// set up the AI Move Request
+		FAIMoveRequest MoveReq;
+
+		MoveReq.SetGoalLocation(Location);
+		MoveReq.SetAcceptanceRadius(AcceptanceRadius);
+		MoveReq.SetAllowPartialPath(true);
+		MoveReq.SetUsePathfinding(true);
+		MoveReq.SetProjectGoalLocation(true);
+		MoveReq.SetRequireNavigableEndLocation(true);
+		MoveReq.SetNavigationFilter(AIController->GetDefaultNavigationFilterClass());
+		MoveReq.SetCanStrafe(false);
+
+		// request a move to the AI Controller
+		FNavPathSharedPtr FollowedPath;
+		const FPathFollowingRequestResult ResultData = AIController->MoveTo(MoveReq, &FollowedPath);
+		
+		// check the move result
+		switch (ResultData.Code)
+		{
+			// failed. Return false
+			case EPathFollowingRequestResult::Failed:
+
+				return false;
+				break;
+
+			// already at goal. Return true and call the move completed delegate
+			case EPathFollowingRequestResult::AlreadyAtGoal:
+
+				OnMoveCompleted.Broadcast(this);
+				return true;
+				break;
+
+			// move successfully scheduled. Return true
+			case EPathFollowingRequestResult::RequestSuccessful:
+
+				return true;
+				break;
+		}
+	}
+
+	// the move could not be completed
+	return false;
+}
+
+void AStrategyUnit::OnMoveFinished(FAIRequestID RequestID, const FPathFollowingResult& Result)
+{
+	// call the delegate
+	OnMoveCompleted.Broadcast(this);
+}
diff --git a/Templates/TP_TopDown/Source/TP_TopDown/Variant_Strategy/StrategyUnit.h b/Templates/TP_TopDown/Source/TP_TopDown/Variant_Strategy/StrategyUnit.h
new file mode 100644
--- /dev/null
+++ b/Templates/TP_TopDown/Source/TP_TopDown/Variant_Strategy/StrategyUnit.h
@@ -0,0 +1,83 @@
+// Copyright Epic Games, Inc. All Rights Reserved.
+
+#pragma once
+
+#include "CoreMinimal.h"
+#include "GameFramework/Character.h"
+#include "AIController.h"
+#include "StrategyUnit.generated.h"
+
+class USphereComponent;
+
+/** Delegate to report that this unit has finished moving */
+DECLARE_DYNAMIC_MULTICAST_DELEGATE_OneParam(FOnUnitMoveCompletedDelegate, AStrategyUnit*, Unit);
+
+/**
+ *  A simple strategy game unit
+ *  Rather than react to inputs, it's controlled indirectly by the Strategy Player Controller
+ */
+UCLASS(abstract)
+class AStrategyUnit : public ACharacter
+{
+	GENERATED_BODY()
+
+private:
+
+	/** Interaction range sphere */
+	UPROPERTY(VisibleAnywhere, BlueprintReadOnly, Category = "Components", meta = (AllowPrivateAccess = "true"))
+	USphereComponent* InteractionRange;
+
+protected:
+
+	/** Cast reference to the AI Controlling this unit */
+	TObjectPtr<AAIController> AIController;
+
+public:
+
+	/** Constructor */
+	AStrategyUnit();
+
+protected:
+
+	virtual void NotifyControllerChanged() override;
+
+public:
+
+	/** Stops unit movement immediately */
+	void StopMoving();
+
+	/** Notifies this unit that it was selected */
+	void UnitSelected();
+
+	/** Notifies this unit that it was deselected */
+	void UnitDeselected();
+
+	/** Notifies this unit that it's been interacted with by another actor */
+	void Interact(AStrategyUnit* Interactor);
+
+	/** Attempts to move this unit to its */
+	bool MoveToLocation(const FVector& Location, float AcceptanceRadius);
+
+protected:
+
+	/** called by the AI controller when this unit has finished moving */
+	void OnMoveFinished(FAIRequestID RequestID, const FPathFollowingResult& Result);
+
+protected:
+
+	/** Blueprint handler for strategy game selection */
+	UFUNCTION(BlueprintImplementableEvent, Category="NPC", meta=(DisplayName="Unit Selected"))
+	void BP_UnitSelected();
+
+	/** Blueprint handler for strategy game deselection */
+	UFUNCTION(BlueprintImplementableEvent, Category="NPC", meta=(DisplayName="Unit Deselected"))
+	void BP_UnitDeselected();
+
+	/** Blueprint handler for strategy game interactions */
+	UFUNCTION(BlueprintImplementableEvent, Category="NPC", meta=(DisplayName="Interaction Behavior"))
+	void BP_InteractionBehavior(AStrategyUnit* Interactor);
+
+public:
+
+	FOnUnitMoveCompletedDelegate OnMoveCompleted;
+};
diff --git a/Templates/TP_TopDown/Source/TP_TopDown/Variant_TwinStick/AI/TwinStickAIController.cpp b/Templates/TP_TopDown/Source/TP_TopDown/Variant_TwinStick/AI/TwinStickAIController.cpp
new file mode 100644
--- /dev/null
+++ b/Templates/TP_TopDown/Source/TP_TopDown/Variant_TwinStick/AI/TwinStickAIController.cpp
@@ -0,0 +1,19 @@
+// Copyright Epic Games, Inc. All Rights Reserved.
+
+
+#include "TwinStickAIController.h"
+#include "Components/StateTreeAIComponent.h"
+
+ATwinStickAIController::ATwinStickAIController()
+{
+	// create the StateTree AI Component
+	StateTreeAI = CreateDefaultSubobject<UStateTreeAIComponent>(TEXT("StateTreeAI"));
+	check(StateTreeAI);
+
+	// ensure we start the StateTree when we possess the pawn
+	bStartAILogicOnPossess = true;
+
+	// ensure we're attached to the possessed character.
+	// this is necessary for EnvQueries to work correctly
+	bAttachToPawn = true;
+}
\ No newline at end of file
diff --git a/Templates/TP_TopDown/Source/TP_TopDown/Variant_TwinStick/AI/TwinStickAIController.h b/Templates/TP_TopDown/Source/TP_TopDown/Variant_TwinStick/AI/TwinStickAIController.h
new file mode 100644
--- /dev/null
+++ b/Templates/TP_TopDown/Source/TP_TopDown/Variant_TwinStick/AI/TwinStickAIController.h
@@ -0,0 +1,28 @@
+// Copyright Epic Games, Inc. All Rights Reserved.
+
+#pragma once
+
+#include "CoreMinimal.h"
+#include "AIController.h"
+#include "TwinStickAIController.generated.h"
+
+class UStateTreeAIComponent;
+
+/**
+ *  A StateTree-Enabled AI Controller for a Twin Stick Shooter game
+ *  Runs NPC logic through a StateTree
+ */
+UCLASS(abstract)
+class ATwinStickAIController : public AAIController
+{
+	GENERATED_BODY()
+	
+	/** StateTree Component */
+	UPROPERTY(VisibleAnywhere, BlueprintReadOnly, Category = "AI", meta = (AllowPrivateAccess = "true"))
+	UStateTreeAIComponent* StateTreeAI;
+
+public:
+
+	/** Constructor */
+	ATwinStickAIController();
+};
diff --git a/Templates/TP_TopDown/Source/TP_TopDown/Variant_TwinStick/AI/TwinStickNPC.cpp b/Templates/TP_TopDown/Source/TP_TopDown/Variant_TwinStick/AI/TwinStickNPC.cpp
new file mode 100644
--- /dev/null
+++ b/Templates/TP_TopDown/Source/TP_TopDown/Variant_TwinStick/AI/TwinStickNPC.cpp
@@ -0,0 +1,126 @@
+// Copyright Epic Games, Inc. All Rights Reserved.
+
+
+#include "TwinStickNPC.h"
+#include "Components/CapsuleComponent.h"
+#include "Components/SkeletalMeshComponent.h"
+#include "TwinStickCharacter.h"
+#include "GameFramework/CharacterMovementComponent.h"
+#include "TwinStickGameMode.h"
+#include "TwinStickPickup.h"
+#include "Engine/World.h"
+#include "TwinStickNPCDestruction.h"
+#include "TimerManager.h"
+
+ATwinStickNPC::ATwinStickNPC()
+{
+	PrimaryActorTick.bCanEverTick = true;
+
+	// ensure we spawn an AI controller when we're spawned
+	AutoPossessAI = EAutoPossessAI::PlacedInWorldOrSpawned;
+
+	// configure the inherited components
+	GetCapsuleComponent()->SetCapsuleRadius(45.0f);
+	GetCapsuleComponent()->SetNotifyRigidBodyCollision(true);
+
+	GetMesh()->SetCollisionProfileName(FName("NoCollision"));
+
+	GetCharacterMovement()->GravityScale = 1.5f;
+	GetCharacterMovement()->MaxAcceleration = 1000.0f;
+	GetCharacterMovement()->BrakingFriction = 1.0f;
+	GetCharacterMovement()->MaxWalkSpeed = 200.0f;
+	GetCharacterMovement()->MaxWalkSpeedCrouched = 100.0f;
+	GetCharacterMovement()->RotationRate = FRotator(0.0f, 640.0f, 0.0f);
+	GetCharacterMovement()->bOrientRotationToMovement = true;
+	GetCharacterMovement()->bUseRVOAvoidance = true;
+	GetCharacterMovement()->AvoidanceConsiderationRadius = 250.0f;
+	GetCharacterMovement()->AvoidanceWeight = 1.0f;
+	GetCharacterMovement()->bConstrainToPlane = true;
+	GetCharacterMovement()->bSnapToPlaneAtStart = true;
+}
+
+void ATwinStickNPC::BeginPlay()
+{
+	Super::BeginPlay();
+
+	// increment the NPC counter so we can cap spawning if necessary
+	if (ATwinStickGameMode* GM = Cast<ATwinStickGameMode>(GetWorld()->GetAuthGameMode()))
+	{
+		GM->IncreaseNPCs();
+	}
+
+}
+
+void ATwinStickNPC::EndPlay(EEndPlayReason::Type EndPlayReason)
+{
+	Super::EndPlay(EndPlayReason);
+
+	// clear the destruction timer
+	GetWorld()->GetTimerManager().ClearTimer(DestructionTimer);
+}
+
+void ATwinStickNPC::Destroyed()
+{
+	// decrease the NPC counter so we can cap spawning if necessary
+	if (ATwinStickGameMode* GM = Cast<ATwinStickGameMode>(GetWorld()->GetAuthGameMode()))
+	{
+		GM->DecreaseNPCs();
+	}
+
+	Super::Destroyed();
+}
+
+void ATwinStickNPC::NotifyHit(class UPrimitiveComponent* MyComp, AActor* Other, class UPrimitiveComponent* OtherComp, bool bSelfMoved, FVector HitLocation, FVector HitNormal, FVector NormalImpulse, const FHitResult& Hit)
+{
+	// have we collided against the player?
+	if (ATwinStickCharacter* PlayerCharacter = Cast<ATwinStickCharacter>(Other))
+	{
+		// apply damage to the character
+		PlayerCharacter->HandleDamage(1.0f, GetActorForwardVector());
+	}
+}
+
+void ATwinStickNPC::ProjectileImpact(const FVector& ForwardVector)
+{
+	// only handle damage if we haven't been hit yet
+	if (bHit)
+	{
+		return;
+	}
+
+	// raise the hit flag
+	bHit = true;
+
+	// deactivate character movement
+	GetCharacterMovement()->Deactivate();
+
+	// award points
+	if (ATwinStickGameMode* GM = Cast<ATwinStickGameMode>(GetWorld()->GetAuthGameMode()))
+	{
+		GM->ScoreUpdate(Score);
+	}
+
+	// randomly spawn a pickup
+	if (FMath::RandRange(0, 100) < PickupSpawnChance)
+	{
+		ATwinStickPickup* Pickup = GetWorld()->SpawnActor<ATwinStickPickup>(PickupClass, GetActorTransform());
+	}
+	
+	// spawn the NPC destruction proxy
+	ATwinStickNPCDestruction* DestructionProxy = GetWorld()->SpawnActor<ATwinStickNPCDestruction>(DestructionProxyClass, GetActorTransform());
+
+	// hide this actor
+	SetActorHiddenInGame(true);
+
+	// disable collision
+	SetActorEnableCollision(false);
+
+	// defer destruction
+	GetWorld()->GetTimerManager().SetTimer(DestructionTimer, this, &ATwinStickNPC::DeferredDestroy, DeferredDestructionTime, false);
+}
+
+void ATwinStickNPC::DeferredDestroy()
+{
+	// destroy this actor
+	Destroy();
+}
diff --git a/Templates/TP_TopDown/Source/TP_TopDown/Variant_TwinStick/AI/TwinStickNPC.h b/Templates/TP_TopDown/Source/TP_TopDown/Variant_TwinStick/AI/TwinStickNPC.h
new file mode 100644
--- /dev/null
+++ b/Templates/TP_TopDown/Source/TP_TopDown/Variant_TwinStick/AI/TwinStickNPC.h
@@ -0,0 +1,81 @@
+// Copyright Epic Games, Inc. All Rights Reserved.
+
+#pragma once
+
+#include "CoreMinimal.h"
+#include "GameFramework/Character.h"
+#include "TwinStickNPC.generated.h"
+
+class ATwinStickPickup;
+class ATwinStickNPCDestruction;
+
+/**
+ *  A simple enemy NPC for a Twin Stick Shooter game
+ *  It's driven by an AI Controller running a behavior tree
+ *  Awards points and randomly spawns pickups on death
+ */
+UCLASS(abstract)
+class ATwinStickNPC : public ACharacter
+{
+	GENERATED_BODY()
+
+protected:
+
+	/** Score to award when this NPC is destroyed */
+	UPROPERTY(EditAnywhere, Category="Score", meta=(ClampMin = 0, ClampMax = 100))
+	int32 Score = 1;
+
+	/** Percentage chance of spawning a pickup */
+	UPROPERTY(EditAnywhere, Category="Pickup", meta=(ClampMin = 0, ClampMax = 100))
+	int32 PickupSpawnChance = 10;
+
+	/** Type of pickup to spawn on death */
+	UPROPERTY(EditAnywhere, Category="Pickup")
+	TSubclassOf<ATwinStickPickup> PickupClass;
+
+	/** Type of destruction proxy to spawn on death */
+	UPROPERTY(EditAnywhere, Category="Destruction")
+	TSubclassOf<ATwinStickNPCDestruction> DestructionProxyClass;
+
+	/** Time to wait after this NPC is hit before destroying it */
+	UPROPERTY(EditAnywhere, Category="Pickup", meta=(ClampMin = 0, ClampMax = 5, Units = "Seconds"))
+	float DeferredDestructionTime = 0.1f;
+
+	/** Deferred destruction timer */
+	FTimerHandle DestructionTimer;
+
+public:
+
+	/** If true, this NPC has already been hit by a projectile and is being destroyed */
+	UPROPERTY(VisibleAnywhere, BlueprintReadOnly, Category="NPC")
+	bool bHit = false;
+
+public:
+
+	/** Constructor */
+	ATwinStickNPC();
+
+protected:
+
+	/** Gameplay Initialization */
+	virtual void BeginPlay() override;
+
+	/** Gameplay cleanup */
+	virtual void EndPlay(EEndPlayReason::Type EndPlayReason) override;
+
+	/** Handle destruction */
+	virtual void Destroyed() override;
+
+	/** Collision handling */
+	virtual void NotifyHit(class UPrimitiveComponent* MyComp, AActor* Other, class UPrimitiveComponent* OtherComp, bool bSelfMoved, FVector HitLocation, FVector HitNormal, FVector NormalImpulse, const FHitResult& Hit) override;
+
+public:
+
+	/** Tells the NPC to process a projectile impact */
+	void ProjectileImpact(const FVector& ForwardVector);
+
+protected:
+
+	/** Called from timer to complete the destruction process for this NPC */
+	void DeferredDestroy();
+};
diff --git a/Templates/TP_TopDown/Source/TP_TopDown/Variant_TwinStick/AI/TwinStickStateTreeUtility.cpp b/Templates/TP_TopDown/Source/TP_TopDown/Variant_TwinStick/AI/TwinStickStateTreeUtility.cpp
new file mode 100644
--- /dev/null
+++ b/Templates/TP_TopDown/Source/TP_TopDown/Variant_TwinStick/AI/TwinStickStateTreeUtility.cpp
@@ -0,0 +1,29 @@
+// Copyright Epic Games, Inc. All Rights Reserved.
+
+
+#include "TwinStickStateTreeUtility.h"
+#include "StateTreeExecutionContext.h"
+#include "StateTreeExecutionTypes.h"
+#include "GameFramework/Character.h"
+#include "Kismet/GameplayStatics.h"
+
+#define LOCTEXT_NAMESPACE "TopDownTemplate"
+
+EStateTreeRunStatus FStateTreeGetPlayerTask::Tick(FStateTreeExecutionContext& Context, const float DeltaTime) const
+{
+	// get the instance data
+	FInstanceDataType& InstanceData = Context.GetInstanceData(*this);
+
+	// get the pawn possessed by the first local player
+	InstanceData.TargetPlayerCharacter = Cast<ACharacter>(UGameplayStatics::GetPlayerPawn(InstanceData.Character, 0));
+
+	// keep the task running
+	return EStateTreeRunStatus::Running;
+}
+
+#if WITH_EDITOR
+FText FStateTreeGetPlayerTask::GetDescription(const FGuid& ID, FStateTreeDataView InstanceDataView, const IStateTreeBindingLookup& BindingLookup, EStateTreeNodeFormatting Formatting /*= EStateTreeNodeFormatting::Text*/) const
+{
+	return LOCTEXT("StateTreeTaskGetPlayerDescription", "<b>Get Player</b>");
+}
+#endif // WITH_EDITOR
\ No newline at end of file
diff --git a/Templates/TP_TopDown/Source/TP_TopDown/Variant_TwinStick/AI/TwinStickStateTreeUtility.h b/Templates/TP_TopDown/Source/TP_TopDown/Variant_TwinStick/AI/TwinStickStateTreeUtility.h
new file mode 100644
--- /dev/null
+++ b/Templates/TP_TopDown/Source/TP_TopDown/Variant_TwinStick/AI/TwinStickStateTreeUtility.h
@@ -0,0 +1,47 @@
+// Copyright Epic Games, Inc. All Rights Reserved.
+
+#pragma once
+
+#include "CoreMinimal.h"
+#include "StateTreeTaskBase.h"
+
+#include "TwinStickStateTreeUtility.generated.h"
+
+class ACharacter;
+
+/**
+ *  Instance data struct for the Get Player task
+ */
+USTRUCT()
+struct FStateTreeGetPlayerInstanceData
+{
+	GENERATED_BODY()
+
+	/** Character that owns this task */
+	UPROPERTY(EditAnywhere, Category = Context)
+	TObjectPtr<ACharacter> Character;
+
+	/** Character that owns this task */
+	UPROPERTY(VisibleAnywhere)
+	TObjectPtr<ACharacter> TargetPlayerCharacter;
+};
+
+/**
+ *  StateTree task to get the player character
+ */
+USTRUCT(meta=(DisplayName="GetPlayer", Category="TwinStick"))
+struct FStateTreeGetPlayerTask : public FStateTreeTaskCommonBase
+{
+	GENERATED_BODY()
+
+	/* Ensure we're using the correct instance data struct */
+	using FInstanceDataType = FStateTreeGetPlayerInstanceData;
+	virtual const UStruct* GetInstanceDataType() const override { return FInstanceDataType::StaticStruct(); }
+
+	/** Runs while the owning state is active */
+	virtual EStateTreeRunStatus Tick(FStateTreeExecutionContext& Context, const float DeltaTime) const override;
+
+#if WITH_EDITOR
+	virtual FText GetDescription(const FGuid& ID, FStateTreeDataView InstanceDataView, const IStateTreeBindingLookup& BindingLookup, EStateTreeNodeFormatting Formatting = EStateTreeNodeFormatting::Text) const override;
+#endif // WITH_EDITOR
+};
\ No newline at end of file
diff --git a/Templates/TP_TopDown/Source/TP_TopDown/Variant_TwinStick/TwinStickAoEAttack.cpp b/Templates/TP_TopDown/Source/TP_TopDown/Variant_TwinStick/TwinStickAoEAttack.cpp
new file mode 100644
--- /dev/null
+++ b/Templates/TP_TopDown/Source/TP_TopDown/Variant_TwinStick/TwinStickAoEAttack.cpp
@@ -0,0 +1,83 @@
+// Copyright Epic Games, Inc. All Rights Reserved.
+
+
+#include "TwinStickAoEAttack.h"
+#include "Components/SceneComponent.h"
+#include "Components/StaticMeshComponent.h"
+#include "Components/SphereComponent.h"
+#include "Engine/World.h"
+#include "TimerManager.h"
+#include "TwinStickNPC.h"
+
+ATwinStickAoEAttack::ATwinStickAoEAttack()
+{
+ 	PrimaryActorTick.bCanEverTick = true;
+
+	// create the root component
+	RootComponent = CreateDefaultSubobject<USceneComponent>(TEXT("Root"));
+
+	// create the mesh that provides the visual representation for the AoE
+	SphereVisual = CreateDefaultSubobject<UStaticMeshComponent>(TEXT("Sphere Visual"));
+	SphereVisual->SetupAttachment(RootComponent);
+
+	SphereVisual->SetCollisionProfileName(FName("NoCollision"));
+
+	// create the collision sphere
+	CollisionSphere = CreateDefaultSubobject<USphereComponent>(TEXT("Collision Sphere"));
+	CollisionSphere->SetupAttachment(RootComponent);
+
+	CollisionSphere->SetSphereRadius(750.0f);
+	CollisionSphere->SetNotifyRigidBodyCollision(true);
+	CollisionSphere->SetCollisionEnabled(ECollisionEnabled::QueryOnly);
+	CollisionSphere->SetCollisionObjectType(ECC_WorldDynamic);
+	CollisionSphere->SetCollisionResponseToAllChannels(ECR_Ignore);
+	CollisionSphere->SetCollisionResponseToChannel(ECC_Pawn, ECR_Overlap);
+}
+
+void ATwinStickAoEAttack::BeginPlay()
+{
+	Super::BeginPlay();
+	
+	// set up the AoE timers
+	GetWorld()->GetTimerManager().SetTimer(TickAoETimer, this, &ATwinStickAoEAttack::TickAoE, TickAoETime, true);
+	GetWorld()->GetTimerManager().SetTimer(StopAoETimer, this, &ATwinStickAoEAttack::StopAoE, StopAoETime, false);
+
+}
+
+void ATwinStickAoEAttack::EndPlay(EEndPlayReason::Type EndPlayReason)
+{
+	Super::EndPlay(EndPlayReason);
+
+	// clear the timers
+	GetWorld()->GetTimerManager().ClearTimer(TickAoETimer);
+	GetWorld()->GetTimerManager().ClearTimer(StopAoETimer);
+}
+
+void ATwinStickAoEAttack::TickAoE()
+{
+	// find all actors overlapping the NPC
+	TArray<AActor*> Overlaps;
+	CollisionSphere->GetOverlappingActors(Overlaps, ATwinStickNPC::StaticClass());
+
+	// process each overlapping actor
+	for (AActor* Current : Overlaps)
+	{
+		if (ATwinStickNPC* NPC = Cast<ATwinStickNPC>(Current))
+		{
+			// tell the NPC it's been hit
+			NPC->ProjectileImpact(FVector::ZeroVector);
+		}
+	}
+}
+
+void ATwinStickAoEAttack::StopAoE()
+{
+	// stop the damage tick timer
+	GetWorld()->GetTimerManager().ClearTimer(TickAoETimer);
+
+	// hide the mesh
+	SphereVisual->SetHiddenInGame(true);
+
+	// call the BP handler. It will be responsible for destroying the Actor when it's done
+	BP_AoEFinished();
+}
diff --git a/Templates/TP_TopDown/Source/TP_TopDown/Variant_TwinStick/TwinStickAoEAttack.h b/Templates/TP_TopDown/Source/TP_TopDown/Variant_TwinStick/TwinStickAoEAttack.h
new file mode 100644
--- /dev/null
+++ b/Templates/TP_TopDown/Source/TP_TopDown/Variant_TwinStick/TwinStickAoEAttack.h
@@ -0,0 +1,69 @@
+// Copyright Epic Games, Inc. All Rights Reserved.
+
+#pragma once
+
+#include "CoreMinimal.h"
+#include "GameFramework/Actor.h"
+#include "TwinStickAoEAttack.generated.h"
+
+class UStaticMeshComponent;
+class USphereComponent;
+
+/**
+ *  A simple persistent AoE attack.
+ *  Damages characters that enter for as long as it's active
+ */
+UCLASS(abstract)
+class ATwinStickAoEAttack : public AActor
+{
+	GENERATED_BODY()
+	
+	/** Provides the visual representation for the AoE attack */
+	UPROPERTY(VisibleAnywhere, BlueprintReadOnly, Category = "Components", meta = (AllowPrivateAccess = "true"))
+	UStaticMeshComponent* SphereVisual;
+
+	/** Provides the collision volume for the AoE attack */
+	UPROPERTY(VisibleAnywhere, BlueprintReadOnly, Category = "Components", meta = (AllowPrivateAccess = "true"))
+	USphereComponent* CollisionSphere;
+
+protected:
+
+	/** Timer to start AoE damage checks */
+	FTimerHandle TickAoETimer;
+
+	/** Timer to end AoE damage checks */
+	FTimerHandle StopAoETimer;
+
+	/** Time to wait between AoE damage ticks */
+	UPROPERTY(EditAnywhere, Category="AoE Attack", meta=(ClampMin = 0, ClampMax = 5, Units = "s"))
+	float TickAoETime = 0.33f;
+
+	/** Time to wait before stopping AoE damage checks */
+	UPROPERTY(EditAnywhere, Category="AoE Attack", meta=(ClampMin = 0, ClampMax = 5, Units = "s"))
+	float StopAoETime = 1.0f;
+
+public:	
+	
+	/** Constructor */
+	ATwinStickAoEAttack();
+
+protected:
+
+	/** Initialization */
+	virtual void BeginPlay() override;
+
+	/** Cleanup */
+	virtual void EndPlay(EEndPlayReason::Type EndPlayReason) override;
+
+protected:
+
+	/** Called when the start AoE timer triggers */
+	void TickAoE();
+
+	/** Called when the stop AoE timer triggers */
+	void StopAoE();
+
+	/** Allows Blueprint handling of AoE fade out effects. NOTE: Call Destroy Actor at the end of this! */
+	UFUNCTION(BlueprintImplementableEvent, Category="AoE Attack")
+	void BP_AoEFinished();
+};
diff --git a/Templates/TP_TopDown/Source/TP_TopDown/Variant_TwinStick/TwinStickCharacter.cpp b/Templates/TP_TopDown/Source/TP_TopDown/Variant_TwinStick/TwinStickCharacter.cpp
new file mode 100644
--- /dev/null
+++ b/Templates/TP_TopDown/Source/TP_TopDown/Variant_TwinStick/TwinStickCharacter.cpp
@@ -0,0 +1,247 @@
+// Copyright Epic Games, Inc. All Rights Reserved.
+
+
+#include "TwinStickCharacter.h"
+#include "GameFramework/SpringArmComponent.h"
+#include "Camera/CameraComponent.h"
+#include "GameFramework/CharacterMovementComponent.h"
+#include "EnhancedInputComponent.h"
+#include "InputAction.h"
+#include "TwinStickGameMode.h"
+#include "TwinStickAoEAttack.h"
+#include "Kismet/KismetMathLibrary.h"
+#include "TwinStickProjectile.h"
+
+ATwinStickCharacter::ATwinStickCharacter()
+{
+ 	PrimaryActorTick.bCanEverTick = true;
+
+	// create the spring arm
+	SpringArm = CreateDefaultSubobject<USpringArmComponent>(TEXT("Spring Arm"));
+	SpringArm->SetupAttachment(RootComponent);
+
+	SpringArm->SetRelativeRotation(FRotator(-50.0f, 0.0f, 0.0f));
+
+	SpringArm->TargetArmLength = 2200.0f;
+	SpringArm->bDoCollisionTest = false;
+	SpringArm->bInheritYaw = false;
+	SpringArm->bEnableCameraLag = true;
+	SpringArm->CameraLagSpeed = 0.5f;
+
+	// create the camera
+	Camera = CreateDefaultSubobject<UCameraComponent>(TEXT("Camera"));
+	Camera->SetupAttachment(SpringArm);
+
+	Camera->SetFieldOfView(75.0f);
+
+	// configure the character movement
+	GetCharacterMovement()->GravityScale = 1.5f;
+	GetCharacterMovement()->MaxAcceleration = 1000.0f;
+	GetCharacterMovement()->BrakingFrictionFactor = 1.0f;
+	GetCharacterMovement()->bCanWalkOffLedges = false;
+	GetCharacterMovement()->RotationRate = FRotator(0.0f, 640.0f, 0.0f);
+	GetCharacterMovement()->bConstrainToPlane = true;
+	GetCharacterMovement()->bSnapToPlaneAtStart = true;
+}
+
+void ATwinStickCharacter::BeginPlay()
+{
+	Super::BeginPlay();
+	
+	// update the items count
+	UpdateItems();
+}
+
+void ATwinStickCharacter::NotifyControllerChanged()
+{
+	Super::NotifyControllerChanged();
+
+	// set the player controller reference
+	PlayerController = Cast<APlayerController>(GetController());
+}
+
+void ATwinStickCharacter::Tick(float DeltaTime)
+{
+	Super::Tick(DeltaTime);
+
+	// get the current rotation
+	const FRotator OldRotation = GetActorRotation();
+
+	// are we aiming with the mouse?
+	if (bUsingMouse)
+	{
+		if (PlayerController)
+		{
+			// get the cursor world location
+			FHitResult OutHit; 
+			PlayerController->GetHitResultUnderCursorByChannel(MouseAimTraceChannel, true, OutHit);
+
+			// find the aim rotation 
+			const FRotator AimRot = UKismetMathLibrary::FindLookAtRotation(GetActorLocation(), OutHit.Location);
+
+			// save the aim angle
+			AimAngle = AimRot.Yaw;
+
+			
+
+			// update the yaw, reuse the pitch and roll
+			SetActorRotation(FRotator(OldRotation.Pitch, AimAngle, OldRotation.Roll));
+
+		}
+
+	} else {
+
+		// use quaternion interpolation to blend between our current rotation
+		// and the desired aim rotation using the shortest path
+		const FRotator TargetRot = FRotator(OldRotation.Pitch, AimAngle, OldRotation.Roll);
+
+		SetActorRotation(TargetRot);
+	}
+}
+
+void ATwinStickCharacter::SetupPlayerInputComponent(UInputComponent* PlayerInputComponent)
+{
+	Super::SetupPlayerInputComponent(PlayerInputComponent);
+
+	// set up the enhanced input action bindings
+	if (UEnhancedInputComponent* EnhancedInputComponent = Cast<UEnhancedInputComponent>(PlayerInputComponent))
+	{
+
+		EnhancedInputComponent->BindAction(MoveAction, ETriggerEvent::Triggered, this, &ATwinStickCharacter::Move);
+		EnhancedInputComponent->BindAction(StickAimAction, ETriggerEvent::Triggered, this, &ATwinStickCharacter::StickAim);
+		EnhancedInputComponent->BindAction(MouseAimAction, ETriggerEvent::Triggered, this, &ATwinStickCharacter::MouseAim);
+		EnhancedInputComponent->BindAction(DashAction, ETriggerEvent::Triggered, this, &ATwinStickCharacter::Dash);
+		EnhancedInputComponent->BindAction(ShootAction, ETriggerEvent::Triggered, this, &ATwinStickCharacter::Shoot);
+		EnhancedInputComponent->BindAction(AoEAction, ETriggerEvent::Triggered, this, &ATwinStickCharacter::AoEAttack);
+
+	}
+
+}
+
+void ATwinStickCharacter::Move(const FInputActionValue& Value)
+{
+	// save the input vector
+	LastMoveInput = Value.Get<FVector2D>();
+
+	// calculate the forward component of the input
+	FRotator FlatRot = GetControlRotation();
+	FlatRot.Pitch = 0.0f;
+
+	// apply the forward input
+	AddMovementInput(FlatRot.RotateVector(FVector::ForwardVector), LastMoveInput.X);
+
+	// apply the right input
+	AddMovementInput(FlatRot.RotateVector(FVector::RightVector), LastMoveInput.Y);
+}
+
+void ATwinStickCharacter::StickAim(const FInputActionValue& Value)
+{
+	// get the input vector
+	FVector2D InputVector = Value.Get<FVector2D>();
+
+	// calculate the aim angle from the inputs
+	AimAngle = FMath::RadiansToDegrees(FMath::Atan2(InputVector.Y, -InputVector.X));
+
+	// lower the mouse controls flag
+	bUsingMouse = false;
+
+	// hide the mouse cursor
+	if (PlayerController)
+	{
+		PlayerController->SetShowMouseCursor(false);
+	}
+}
+
+void ATwinStickCharacter::MouseAim(const FInputActionValue& Value)
+{
+	// raise the mouse controls flag
+	bUsingMouse = true;
+
+	// show the mouse cursor
+	if (PlayerController)
+	{
+		PlayerController->SetShowMouseCursor(true);
+	}
+}
+
+void ATwinStickCharacter::Dash(const FInputActionValue& Value)
+{
+	// calculate the launch impulse vector based on the last move input
+	FVector LaunchDir = FVector::ZeroVector;
+
+	LaunchDir.X = FMath::Clamp(LastMoveInput.X, -1.0f, 1.0f);
+	LaunchDir.Y = FMath::Clamp(LastMoveInput.Y, -1.0f, 1.0f);
+
+	// launch the character in the chosen direction
+	LaunchCharacter(LaunchDir * DashImpulse, true, true);
+}
+
+void ATwinStickCharacter::Shoot(const FInputActionValue& Value)
+{
+	// get the actor transform
+	FTransform ProjectileTransform = GetActorTransform();
+
+	// apply the projectile spawn offset
+	FVector ProjectileLocation = ProjectileTransform.GetLocation() + ProjectileTransform.GetRotation().RotateVector(FVector::ForwardVector * ProjectileOffset);
+	ProjectileTransform.SetLocation(ProjectileLocation);
+
+	ATwinStickProjectile* Projectile = GetWorld()->SpawnActor<ATwinStickProjectile>(ProjectileClass, ProjectileTransform);
+}
+
+void ATwinStickCharacter::AoEAttack(const FInputActionValue& Value)
+{
+	// do we have enough items to do an AoE attack?
+	if (Items > 0)
+	{
+		// get the game time
+		const float GameTime = GetWorld()->GetTimeSeconds();
+
+		// are we off AoE cooldown?
+		if (GameTime - LastAoETime > AoECooldownTime)
+		{
+			// save the new AoE time
+			LastAoETime = GameTime;
+
+			// spawn the AoE
+			ATwinStickAoEAttack* AoE = GetWorld()->SpawnActor<ATwinStickAoEAttack>(AoEAttackClass, GetActorTransform());
+
+			// decrease the number of items
+			--Items;
+
+			// update the items count
+			UpdateItems();
+		}
+	}
+}
+
+void ATwinStickCharacter::HandleDamage(float Damage, const FVector& DamageDirection)
+{
+	// calculate the knockback vector
+	FVector LaunchVector = DamageDirection;
+	LaunchVector.Z = 0.0f;
+
+	// apply knockback to the character
+	LaunchCharacter(LaunchVector * KnockbackStrength, true, true);
+
+	// pass control to BP
+	BP_Damaged();
+}
+
+void ATwinStickCharacter::AddPickup()
+{
+	// increase the item count
+	++Items;
+
+	// update the items counter
+	UpdateItems();
+}
+
+void ATwinStickCharacter::UpdateItems()
+{
+	// update the game mode
+	if (ATwinStickGameMode* GM = Cast<ATwinStickGameMode>(GetWorld()->GetAuthGameMode()))
+	{
+		GM->ItemUsed(Items);
+	}
+}
+
diff --git a/Templates/TP_TopDown/Source/TP_TopDown/Variant_TwinStick/TwinStickCharacter.h b/Templates/TP_TopDown/Source/TP_TopDown/Variant_TwinStick/TwinStickCharacter.h
new file mode 100644
--- /dev/null
+++ b/Templates/TP_TopDown/Source/TP_TopDown/Variant_TwinStick/TwinStickCharacter.h
@@ -0,0 +1,174 @@
+// Copyright Epic Games, Inc. All Rights Reserved.
+
+#pragma once
+
+#include "CoreMinimal.h"
+#include "GameFramework/Character.h"
+#include "GameFramework/PlayerController.h"
+#include "TwinStickCharacter.generated.h"
+
+class USpringArmComponent;
+class UCameraComponent;
+struct FInputActionValue;
+class APlayerController;
+class UInputAction;
+class ATwinStickAoEAttack;
+class ATwinStickProjectile;
+
+/**
+ *  A player-controlled character for a Twin Stick Shooter game
+ *  Automatically rotates to face the aim direction.
+ *  Fires projectiles and spawns AoE attacks.
+ */
+UCLASS(abstract)
+class ATwinStickCharacter : public ACharacter
+{
+	GENERATED_BODY()
+
+	/** Camera boom spring arm */
+	UPROPERTY(VisibleAnywhere, BlueprintReadOnly, Category = "Components", meta = (AllowPrivateAccess = "true"))
+	USpringArmComponent* SpringArm;
+
+	/** Player Camera */
+	UPROPERTY(VisibleAnywhere, BlueprintReadOnly, Category = "Components", meta = (AllowPrivateAccess = "true"))
+	UCameraComponent* Camera;
+
+protected:
+
+	/** Movement input action */
+	UPROPERTY(EditAnywhere, Category="Input")
+	UInputAction* MoveAction;
+
+	/** Gamepad aim input action */
+	UPROPERTY(EditAnywhere, Category="Input")
+	UInputAction* StickAimAction;
+
+	/** Mouse aim input action */
+	UPROPERTY(EditAnywhere, Category="Input")
+	UInputAction* MouseAimAction;
+
+	/** Dash input action */
+	UPROPERTY(EditAnywhere, Category="Input")
+	UInputAction* DashAction;
+
+	/** Shooting input action */
+	UPROPERTY(EditAnywhere, Category="Input")
+	UInputAction* ShootAction;
+
+	/** AoE attack input action */
+	UPROPERTY(EditAnywhere, Category="Input")
+	UInputAction* AoEAction;
+
+	/** Trace channel to use for mouse aim */
+	UPROPERTY(EditAnywhere, Category="Input")
+	TEnumAsByte<ETraceTypeQuery> MouseAimTraceChannel;
+
+	/** Impulse to apply to the character when dashing */
+	UPROPERTY(EditAnywhere, Category="Dash")
+	float DashImpulse = 2500.0f;
+
+	/** Type of projectile to spawn when shooting */
+	UPROPERTY(EditAnywhere, Category="Projectile")
+	TSubclassOf<ATwinStickProjectile> ProjectileClass;
+
+	/** Distance ahead of the character that the projectile will be spawned at */
+	UPROPERTY(EditAnywhere, Category="Projectile", meta=(ClampMin = 0, ClampMax = 1000, Units = "cm"))
+	float ProjectileOffset = 100.0f;
+
+	/** Type of AoE attack actor to spawn */
+	UPROPERTY(EditAnywhere, Category="AoE")
+	TSubclassOf<ATwinStickAoEAttack> AoEAttackClass;
+
+	/** Number of starting AoE attack items */
+	UPROPERTY(EditAnywhere, Category="AoE")
+	int32 Items = 1;
+
+	/** Knockback impulse to apply to the character when they're damaged */
+	UPROPERTY(EditAnywhere, Category="Damage", meta=(ClampMin = 0, ClampMax = 1000, Units = "cm"))
+	float KnockbackStrength = 2500.0f;
+
+	/** Time to disallow AoE attacks after one is performed */
+	UPROPERTY(EditAnywhere, Category="AoE", meta=(ClampMin = 0, ClampMax = 10, Units = "s"))
+	float AoECooldownTime = 1.0f;
+
+	/** Speed to blend between our current rotation and the target aim rotation when stick aiming */
+	UPROPERTY(EditAnywhere, Category="Aim", meta=(ClampMin = 0, ClampMax = 100))
+	float AimRotationInterpSpeed = 10.0f;
+
+	/** Game time of the last AoE attack */
+	float LastAoETime = 0.0f;
+
+	/** Aim Yaw Angle in degrees */
+	float AimAngle = 0.0f;
+
+	/** Pointer to the player controller assigned to this character */
+	TObjectPtr<APlayerController> PlayerController;
+
+	/** If true, the player is using mouse aim */
+	bool bUsingMouse = false;
+
+	/** Last held move input */
+	FVector2D LastMoveInput;
+
+public:
+	
+	/** Constructor */
+	ATwinStickCharacter();
+
+protected:
+
+	/** Gameplay Initialization */
+	virtual void BeginPlay() override;
+
+	/** Possessed by controller initialization */
+	virtual void NotifyControllerChanged() override;
+
+public:	
+	
+	/** Updates the character's rotation to face the aim direction */
+	virtual void Tick(float DeltaTime) override;
+
+	/** Adds input bindings */
+	virtual void SetupPlayerInputComponent(class UInputComponent* PlayerInputComponent) override;
+
+protected:
+
+	/** Handles movement inputs */
+	void Move(const FInputActionValue& Value);
+
+	/** Handles joypad aim */
+	void StickAim(const FInputActionValue& Value);
+
+	/** Handles mouse aim */
+	void MouseAim(const FInputActionValue& Value);
+
+	/** Performs a dash */
+	void Dash(const FInputActionValue& Value);
+
+	/** Shoots projectiles */
+	void Shoot(const FInputActionValue& Value);
+
+	/** Performs an AoE Attack */
+	void AoEAttack(const FInputActionValue& Value);
+
+public:
+
+	/** Applies collision impact to the player */
+	void HandleDamage(float Damage, const FVector& DamageDirection);
+
+protected:
+
+	/** Allows Blueprint code to react to damage */
+	UFUNCTION(BlueprintImplementableEvent, Category="Damage", meta=(DisplayName = "Damaged"))
+	void BP_Damaged();
+
+public:
+
+	/** Gives the player a pickup item */
+	void AddPickup();
+
+protected:
+
+	/** Updates the items counter on the Game Mode */
+	void UpdateItems();
+};
diff --git a/Templates/TP_TopDown/Source/TP_TopDown/Variant_TwinStick/TwinStickGameMode.cpp b/Templates/TP_TopDown/Source/TP_TopDown/Variant_TwinStick/TwinStickGameMode.cpp
new file mode 100644
--- /dev/null
+++ b/Templates/TP_TopDown/Source/TP_TopDown/Variant_TwinStick/TwinStickGameMode.cpp
@@ -0,0 +1,113 @@
+// Copyright Epic Games, Inc. All Rights Reserved.
+
+
+#include "TwinStickGameMode.h"
+#include "TwinStickUI.h"
+#include "Engine/World.h"
+#include "TimerManager.h"
+#include "Kismet/GameplayStatics.h"
+
+void ATwinStickGameMode::BeginPlay()
+{
+	// create the UI widget and add it to the viewport
+	UIWidget = CreateWidget<UTwinStickUI>(UGameplayStatics::GetPlayerController(GetWorld(), 0), UIWidgetClass);
+	UIWidget->AddToViewport(0);
+}
+
+void ATwinStickGameMode::EndPlay(EEndPlayReason::Type EndPlayReason)
+{
+	Super::EndPlay(EndPlayReason);
+	
+	// clear the combo timer
+	GetWorld()->GetTimerManager().ClearTimer(ComboTimer);
+}
+
+void ATwinStickGameMode::ItemUsed(int32 Value)
+{
+	// update the UI
+	UIWidget->UpdateItems(Value);
+}
+
+void ATwinStickGameMode::ScoreUpdate(int32 Value)
+{
+	// multiply the base score by the combo multiplier and add it to the score
+	Score += Value * Combo;
+
+	// update the UI
+	UIWidget->UpdateScore(Score);
+
+	// update the combo multiplier
+	ComboUpdate();
+}
+
+void ATwinStickGameMode::ComboUpdate()
+{
+	// return
+	if (Combo > ComboCap)
+	{
+		return;
+	}
+
+	// update the combo increment
+	++ComboIncrement;
+
+	// is it time to increase the multiplier?
+	if (ComboIncrement > ComboIncrementMax)
+	{
+		// reset the combo increment
+		ComboIncrement = 0;
+
+		// increase the combo multiplier
+		++Combo;
+
+		// update the UI
+		UIWidget->UpdateCombo(Combo);
+
+	}
+
+	// reset the cooldown timer
+	ResetComboCooldown();
+}
+
+void ATwinStickGameMode::ResetComboCooldown()
+{
+	// reset the combo cooldown timer
+	GetWorld()->GetTimerManager().SetTimer(ComboTimer, this, &ATwinStickGameMode::ResetCombo, ComboCooldown, false);
+}
+
+void ATwinStickGameMode::ResetCombo()
+{
+	// is the combo multiplier above min?
+	if (Combo > 1)
+	{
+		// reset the combo increment
+		ComboIncrement = 0;
+
+		// tick down the multiplier
+		--Combo;
+
+		// update the UI
+		UIWidget->UpdateCombo(Combo);
+
+		// reset the cooldown timer
+		ResetComboCooldown();
+	}
+}
+
+bool ATwinStickGameMode::CanSpawnNPCs()
+{
+	// is the NPC counter under the cap?
+	return NPCCount < NPCCap;
+}
+
+void ATwinStickGameMode::IncreaseNPCs()
+{
+	// increase the NPC counter
+	++NPCCount;
+}
+
+void ATwinStickGameMode::DecreaseNPCs()
+{
+	// decrease the NPC counter
+	--NPCCount;
+}
diff --git a/Templates/TP_TopDown/Source/TP_TopDown/Variant_TwinStick/TwinStickGameMode.h b/Templates/TP_TopDown/Source/TP_TopDown/Variant_TwinStick/TwinStickGameMode.h
new file mode 100644
--- /dev/null
+++ b/Templates/TP_TopDown/Source/TP_TopDown/Variant_TwinStick/TwinStickGameMode.h
@@ -0,0 +1,99 @@
+// Copyright Epic Games, Inc. All Rights Reserved.
+
+#pragma once
+
+#include "CoreMinimal.h"
+#include "GameFramework/GameModeBase.h"
+#include "TwinStickGameMode.generated.h"
+
+class UTwinStickUI;
+
+/**
+ *  Simple Game Mode for a Twin Stick Shooter game.
+ *  Manages the score and UI
+ */
+UCLASS(abstract)
+class ATwinStickGameMode : public AGameModeBase
+{
+	GENERATED_BODY()
+	
+protected:
+
+	/** Type of UI Widget to spawn */
+	UPROPERTY(EditAnywhere, Category="Twin Stick")
+	TSubclassOf<UTwinStickUI> UIWidgetClass;
+
+	/** Pointer to the spawned UI Widget */
+	TObjectPtr<UTwinStickUI> UIWidget;
+
+	/** Current game score */
+	int32 Score = 0;
+
+	/** Current combo multiplier */
+	int32 Combo = 1;
+
+	/** Current combo increment value */
+	int32 ComboIncrement = 0;
+
+	/** Number of combo hits to process before incrementing the combo multiplier */
+	UPROPERTY(EditAnywhere, Category="Twin Stick", meta=(ClampMin = 0, ClampMax = 10))
+	int32 ComboIncrementMax = 5;
+
+	/** Maximum allowed combo multiplier value */
+	UPROPERTY(EditAnywhere, Category="Twin Stick", meta=(ClampMin = 0, ClampMax = 10))
+	int32 ComboCap = 4;
+
+	/** Max time between kills before the combo multiplier resets */
+	UPROPERTY(EditAnywhere, Category="Twin Stick", meta=(ClampMin = 0, ClampMax = 10, Units = "s"))
+	float ComboCooldown = 3.0f;
+
+	/** Game time of the last combo kill */
+	float LastComboTime = 0.0f;
+
+	FTimerHandle ComboTimer;
+
+	/** Max number of NPCs to allow in the level at once */
+	UPROPERTY(EditAnywhere, Category="Twin Stick", meta=(ClampMin = 0, ClampMax = 100))
+	int32 NPCCap = 20;
+
+	/** Current number of NPCs in the level */
+	int32 NPCCount = 0;
+
+public:
+
+	/** Gameplay initialization */
+	virtual void BeginPlay() override;
+
+	/** Cleanup */
+	virtual void EndPlay(EEndPlayReason::Type EndPlayReason) override;
+
+public:
+
+	/** Called when an item has been used */
+	void ItemUsed(int32 Value);
+
+	/** Increments the score by the given value */
+	void ScoreUpdate(int32 Value);
+
+protected:
+
+	/** Updates the combo multiplier */
+	void ComboUpdate();
+
+	/** Resets the combo cooldown timer */
+	void ResetComboCooldown();
+
+	/** Resets the combo multiplier after the cooldown time expires */
+	void ResetCombo();
+
+public:
+
+	/** Returns true if the number of NPCs is under the cap */
+	bool CanSpawnNPCs();
+
+	/** Increases the NPC count */
+	void IncreaseNPCs();
+
+	/** Decreases the NPC count */
+	void DecreaseNPCs();
+};
diff --git a/Templates/TP_TopDown/Source/TP_TopDown/Variant_TwinStick/TwinStickNPCDestruction.cpp b/Templates/TP_TopDown/Source/TP_TopDown/Variant_TwinStick/TwinStickNPCDestruction.cpp
new file mode 100644
--- /dev/null
+++ b/Templates/TP_TopDown/Source/TP_TopDown/Variant_TwinStick/TwinStickNPCDestruction.cpp
@@ -0,0 +1,10 @@
+// Copyright Epic Games, Inc. All Rights Reserved.
+
+
+#include "TwinStickNPCDestruction.h"
+
+ATwinStickNPCDestruction::ATwinStickNPCDestruction()
+{
+ 	PrimaryActorTick.bCanEverTick = true;
+
+}
\ No newline at end of file
diff --git a/Templates/TP_TopDown/Source/TP_TopDown/Variant_TwinStick/TwinStickNPCDestruction.h b/Templates/TP_TopDown/Source/TP_TopDown/Variant_TwinStick/TwinStickNPCDestruction.h
new file mode 100644
--- /dev/null
+++ b/Templates/TP_TopDown/Source/TP_TopDown/Variant_TwinStick/TwinStickNPCDestruction.h
@@ -0,0 +1,24 @@
+// Copyright Epic Games, Inc. All Rights Reserved.
+
+#pragma once
+
+#include "CoreMinimal.h"
+#include "GameFramework/Actor.h"
+#include "TwinStickNPCDestruction.generated.h"
+
+/**
+ *  A NPC destruction proxy for a Twin Stick Shooter game
+ *  Replaces the NPC when it is destroyed,
+ *  allowing it to play effects without affecting gameplay 
+ */
+UCLASS(abstract)
+class ATwinStickNPCDestruction : public AActor
+{
+	GENERATED_BODY()
+	
+public:
+
+	/** Constructor */
+	ATwinStickNPCDestruction();
+
+};
diff --git a/Templates/TP_TopDown/Source/TP_TopDown/Variant_TwinStick/TwinStickPickup.cpp b/Templates/TP_TopDown/Source/TP_TopDown/Variant_TwinStick/TwinStickPickup.cpp
new file mode 100644
--- /dev/null
+++ b/Templates/TP_TopDown/Source/TP_TopDown/Variant_TwinStick/TwinStickPickup.cpp
@@ -0,0 +1,49 @@
+// Copyright Epic Games, Inc. All Rights Reserved.
+
+
+#include "TwinStickPickup.h"
+#include "Components/SceneComponent.h"
+#include "Components/SphereComponent.h"
+#include "TwinStickCharacter.h"
+#include "Components/StaticMeshComponent.h"
+
+ATwinStickPickup::ATwinStickPickup()
+{
+ 	PrimaryActorTick.bCanEverTick = true;
+
+	// create the root component
+	RootComponent = CreateDefaultSubobject<USceneComponent>(TEXT("Root"));
+
+	// create the collision sphere
+	CollisionSphere = CreateDefaultSubobject<USphereComponent>(TEXT("Collision Sphere"));
+	CollisionSphere->SetupAttachment(RootComponent);
+
+	CollisionSphere->SetSphereRadius(100.0f);
+	CollisionSphere->SetRelativeLocation(FVector(0.0f, 0.0f, 125.0f));
+	CollisionSphere->SetCollisionEnabled(ECollisionEnabled::QueryOnly);
+	CollisionSphere->SetCollisionObjectType(ECC_WorldDynamic);
+	CollisionSphere->SetCollisionResponseToAllChannels(ECR_Ignore);
+	CollisionSphere->SetCollisionResponseToChannel(ECC_Pawn, ECR_Overlap);
+
+	// create the mesh
+	Mesh = CreateDefaultSubobject<UStaticMeshComponent>(TEXT("Mesh"));
+	Mesh->SetupAttachment(CollisionSphere);
+
+	Mesh->SetCollisionProfileName(FName("NoCollision"));
+
+}
+
+void ATwinStickPickup::NotifyActorBeginOverlap(AActor* OtherActor)
+{
+	Super::NotifyActorBeginOverlap(OtherActor);
+
+	// have we overlapped the player character?
+	if (ATwinStickCharacter* PlayerCharacter = Cast<ATwinStickCharacter>(OtherActor))
+	{
+		// give the pickup to the player
+		PlayerCharacter->AddPickup();
+
+		// destroy this pickup
+		Destroy();
+	}
+}
\ No newline at end of file
diff --git a/Templates/TP_TopDown/Source/TP_TopDown/Variant_TwinStick/TwinStickPickup.h b/Templates/TP_TopDown/Source/TP_TopDown/Variant_TwinStick/TwinStickPickup.h
new file mode 100644
--- /dev/null
+++ b/Templates/TP_TopDown/Source/TP_TopDown/Variant_TwinStick/TwinStickPickup.h
@@ -0,0 +1,36 @@
+// Copyright Epic Games, Inc. All Rights Reserved.
+
+#pragma once
+
+#include "CoreMinimal.h"
+#include "GameFramework/Actor.h"
+#include "TwinStickPickup.generated.h"
+
+class USphereComponent;
+class UStaticMeshComponent;
+
+/**
+ *  A simple pickup for a Twin Stick Shooter game
+ */
+UCLASS(abstract)
+class ATwinStickPickup : public AActor
+{
+	GENERATED_BODY()
+	
+	/** Pickup collision sphere */
+	UPROPERTY(VisibleAnywhere, BlueprintReadOnly, Category = "Components", meta = (AllowPrivateAccess = "true"))
+	USphereComponent* CollisionSphere;
+
+	/** Provides visual representation for the pickup */
+	UPROPERTY(VisibleAnywhere, BlueprintReadOnly, Category = "Components", meta = (AllowPrivateAccess = "true"))
+	UStaticMeshComponent* Mesh;
+
+public:	
+
+	/** Constructor */
+	ATwinStickPickup();
+
+	/** Collision handling */
+	virtual void NotifyActorBeginOverlap(AActor* OtherActor) override;
+
+};
diff --git a/Templates/TP_TopDown/Source/TP_TopDown/Variant_TwinStick/TwinStickPlayerController.cpp b/Templates/TP_TopDown/Source/TP_TopDown/Variant_TwinStick/TwinStickPlayerController.cpp
new file mode 100644
--- /dev/null
+++ b/Templates/TP_TopDown/Source/TP_TopDown/Variant_TwinStick/TwinStickPlayerController.cpp
@@ -0,0 +1,47 @@
+// Copyright Epic Games, Inc. All Rights Reserved.
+
+
+#include "TwinStickPlayerController.h"
+#include "EnhancedInputSubsystems.h"
+#include "InputMappingContext.h"
+#include "Kismet/GameplayStatics.h"
+#include "GameFramework/PlayerStart.h"
+#include "TwinStickCharacter.h"
+#include "Engine/LocalPlayer.h"
+#include "Engine/World.h"
+
+void ATwinStickPlayerController::SetupInputComponent()
+{
+	// add the input mapping context
+	if (UEnhancedInputLocalPlayerSubsystem* Subsystem = ULocalPlayer::GetSubsystem<UEnhancedInputLocalPlayerSubsystem>(GetLocalPlayer()))
+	{
+		Subsystem->AddMappingContext(InputMappingContext, 0);
+	}
+}
+
+void ATwinStickPlayerController::OnPossess(APawn* InPawn)
+{
+	Super::OnPossess(InPawn);
+
+	// subscribe to the pawn's OnDestroyed delegate
+	InPawn->OnDestroyed.AddDynamic(this, &ATwinStickPlayerController::OnPawnDestroyed);
+}
+
+void ATwinStickPlayerController::OnPawnDestroyed(AActor* DestroyedActor)
+{
+	// find the player start
+	TArray<AActor*> ActorList;
+	UGameplayStatics::GetAllActorsOfClass(GetWorld(), APlayerStart::StaticClass(), ActorList);
+
+	if (ActorList.Num() > 0)
+	{
+		// spawn a character at the player start
+		const FTransform SpawnTransform = ActorList[0]->GetActorTransform();
+
+		if (ATwinStickCharacter* RespawnedCharacter = GetWorld()->SpawnActor<ATwinStickCharacter>(CharacterClass, SpawnTransform))
+		{
+			// possess the character
+			Possess(RespawnedCharacter);
+		}
+	}
+}
diff --git a/Templates/TP_TopDown/Source/TP_TopDown/Variant_TwinStick/TwinStickPlayerController.h b/Templates/TP_TopDown/Source/TP_TopDown/Variant_TwinStick/TwinStickPlayerController.h
new file mode 100644
--- /dev/null
+++ b/Templates/TP_TopDown/Source/TP_TopDown/Variant_TwinStick/TwinStickPlayerController.h
@@ -0,0 +1,43 @@
+// Copyright Epic Games, Inc. All Rights Reserved.
+
+#pragma once
+
+#include "CoreMinimal.h"
+#include "GameFramework/PlayerController.h"
+#include "TwinStickPlayerController.generated.h"
+
+class UInputMappingContext;
+class ATwinStickCharacter;
+
+/**
+ *  Simple Player Controller for a Twin Stick Shooter game
+ *  Manages input mapping contexts
+ *  Respawns the pawn if it is destroyed
+ */
+UCLASS(abstract)
+class ATwinStickPlayerController : public APlayerController
+{
+	GENERATED_BODY()
+	
+protected:
+
+	/** Input mapping context for this player */
+	UPROPERTY(EditAnywhere, Category="Input")
+	UInputMappingContext* InputMappingContext;
+
+	/** Character class to respawn when the possessed pawn is destroyed */
+	UPROPERTY(EditAnywhere, Category="Respawn")
+	TSubclassOf<ATwinStickCharacter> CharacterClass;
+
+protected:
+
+	/** Initialize input bindings */
+	virtual void SetupInputComponent() override;
+
+	/** Pawn initialization */
+	virtual void OnPossess(APawn* InPawn) override;
+
+	/** Called if the possessed pawn is destroyed */
+	UFUNCTION()
+	void OnPawnDestroyed(AActor* DestroyedActor);
+};
diff --git a/Templates/TP_TopDown/Source/TP_TopDown/Variant_TwinStick/TwinStickProjectile.cpp b/Templates/TP_TopDown/Source/TP_TopDown/Variant_TwinStick/TwinStickProjectile.cpp
new file mode 100644
--- /dev/null
+++ b/Templates/TP_TopDown/Source/TP_TopDown/Variant_TwinStick/TwinStickProjectile.cpp
@@ -0,0 +1,65 @@
+// Copyright Epic Games, Inc. All Rights Reserved.
+
+
+#include "TwinStickProjectile.h"
+#include "Components/SphereComponent.h"
+#include "GameFramework/ProjectileMovementComponent.h"
+#include "Components/StaticMeshComponent.h"
+#include "TwinStickNPC.h"
+
+ATwinStickProjectile::ATwinStickProjectile()
+{
+ 	PrimaryActorTick.bCanEverTick = true;
+
+	// this actor will be destroyed automatically once InitialLifeSpan expires
+	InitialLifeSpan = 2.0f;
+
+	// create the collision sphere and set it as the root component
+	RootComponent = CollisionSphere = CreateDefaultSubobject<USphereComponent>(TEXT("Collision Sphere"));
+
+	CollisionSphere->SetSphereRadius(35.0f);
+	CollisionSphere->SetNotifyRigidBodyCollision(true);
+	CollisionSphere->SetCollisionEnabled(ECollisionEnabled::QueryOnly);
+	CollisionSphere->SetCollisionObjectType(ECC_WorldDynamic);
+	CollisionSphere->SetCollisionResponseToAllChannels(ECR_Block);
+
+	// create the mesh
+	Mesh = CreateDefaultSubobject<UStaticMeshComponent>(TEXT("Mesh"));
+	Mesh->SetupAttachment(RootComponent);
+
+	Mesh->SetCollisionProfileName(FName("NoCollision"));
+
+	// create the projectile movement comp. No need to attach it because it's not a scene component
+	ProjectileMovement = CreateDefaultSubobject<UProjectileMovementComponent>(TEXT("Projectile Movement"));
+
+	ProjectileMovement->InitialSpeed = 2000.0f;
+	ProjectileMovement->MaxSpeed = 15000.0f;
+	ProjectileMovement->bRotationFollowsVelocity = true;
+	ProjectileMovement->bRotationRemainsVertical = true;
+	ProjectileMovement->ProjectileGravityScale = 0.0f;
+	ProjectileMovement->bShouldBounce = true;
+	ProjectileMovement->bForceSubStepping = true;
+
+	ProjectileMovement->OnProjectileStop.AddDynamic(this, &ATwinStickProjectile::OnProjectileStop);
+}
+
+void ATwinStickProjectile::NotifyHit(class UPrimitiveComponent* MyComp, AActor* Other, class UPrimitiveComponent* OtherComp, bool bSelfMoved, FVector HitLocation, FVector HitNormal, FVector NormalImpulse, const FHitResult& Hit)
+{
+	Super::NotifyHit(MyComp, Other, OtherComp, bSelfMoved, HitLocation, HitNormal, NormalImpulse, Hit);
+
+	// have we hit a NPC?
+	if (ATwinStickNPC* NPC = Cast<ATwinStickNPC>(Other))
+	{
+		// tell the NPC it's been hit
+		NPC->ProjectileImpact(FVector::ZeroVector);
+
+		// destroy this projectile
+		Destroy();
+	}
+}
+
+void ATwinStickProjectile::OnProjectileStop(const FHitResult& ImpactResult)
+{
+	// destroy this actor immediately
+	Destroy();
+}
diff --git a/Templates/TP_TopDown/Source/TP_TopDown/Variant_TwinStick/TwinStickProjectile.h b/Templates/TP_TopDown/Source/TP_TopDown/Variant_TwinStick/TwinStickProjectile.h
new file mode 100644
--- /dev/null
+++ b/Templates/TP_TopDown/Source/TP_TopDown/Variant_TwinStick/TwinStickProjectile.h
@@ -0,0 +1,47 @@
+// Copyright Epic Games, Inc. All Rights Reserved.
+
+#pragma once
+
+#include "CoreMinimal.h"
+#include "GameFramework/Actor.h"
+#include "TwinStickProjectile.generated.h"
+
+class USphereComponent;
+class UStaticMeshComponent;
+class UProjectileMovementComponent;
+
+/**
+ *  A simple bouncing projectile for a Twin Stick shooter game
+ */
+UCLASS(abstract)
+class ATwinStickProjectile : public AActor
+{
+	GENERATED_BODY()
+	
+	/** Projectile collision sphere */
+	UPROPERTY(VisibleAnywhere, BlueprintReadOnly, Category = "Components", meta = (AllowPrivateAccess = "true"))
+	USphereComponent* CollisionSphere;
+
+	/** Mesh that provides the visual representation for this projectile */
+	UPROPERTY(VisibleAnywhere, BlueprintReadOnly, Category = "Components", meta = (AllowPrivateAccess = "true"))
+	UStaticMeshComponent* Mesh;
+
+	/** Handles movement behaviors for this projectile */
+	UPROPERTY(VisibleAnywhere, BlueprintReadOnly, Category = "Components", meta = (AllowPrivateAccess = "true"))
+	UProjectileMovementComponent* ProjectileMovement;
+
+public:	
+
+	/** Constructor */
+	ATwinStickProjectile();
+
+	/** Handles collisions */
+	virtual void NotifyHit(class UPrimitiveComponent* MyComp, AActor* Other, class UPrimitiveComponent* OtherComp, bool bSelfMoved, FVector HitLocation, FVector HitNormal, FVector NormalImpulse, const FHitResult& Hit) override;
+
+protected:
+	
+	/** Handles collisions that stop this projectile from moving */
+	UFUNCTION()
+	void OnProjectileStop(const FHitResult& ImpactResult);
+
+};
diff --git a/Templates/TP_TopDown/Source/TP_TopDown/Variant_TwinStick/TwinStickSpawner.cpp b/Templates/TP_TopDown/Source/TP_TopDown/Variant_TwinStick/TwinStickSpawner.cpp
new file mode 100644
--- /dev/null
+++ b/Templates/TP_TopDown/Source/TP_TopDown/Variant_TwinStick/TwinStickSpawner.cpp
@@ -0,0 +1,90 @@
+// Copyright Epic Games, Inc. All Rights Reserved.
+
+
+#include "TwinStickSpawner.h"
+#include "Engine/World.h"
+#include "TimerManager.h"
+#include "NavigationSystem.h"
+#include "NavMesh/RecastNavMesh.h"
+#include "Kismet/GameplayStatics.h"
+#include "TwinStickNPC.h"
+#include "TwinStickGameMode.h"
+
+ATwinStickSpawner::ATwinStickSpawner()
+{
+ 	PrimaryActorTick.bCanEverTick = true;
+
+}
+
+void ATwinStickSpawner::BeginPlay()
+{
+	Super::BeginPlay();
+	
+	// find the recast navmesh actor on the level
+	TArray<AActor*> ActorList;
+	UGameplayStatics::GetAllActorsOfClass(GetWorld(), ARecastNavMesh::StaticClass(), ActorList);
+
+	if (ActorList.Num() > 0)
+	{
+		NavData = Cast<ARecastNavMesh>(ActorList[0]);
+	} else {
+
+		UE_LOG(LogTemp, Log, TEXT("Could not find recast navmesh"));
+
+	}
+
+	// set up the spawn timer
+	GetWorld()->GetTimerManager().SetTimer(SpawnGroupTimer, this, &ATwinStickSpawner::SpawnNPCGroup, SpawnGroupDelay, true);
+
+	// spawn the first group of NPCs
+	SpawnNPCGroup();
+}
+
+void ATwinStickSpawner::EndPlay(EEndPlayReason::Type EndPlayReason)
+{
+	Super::EndPlay(EndPlayReason);
+
+	// clear the spawn timers
+	GetWorld()->GetTimerManager().ClearTimer(SpawnGroupTimer);
+	GetWorld()->GetTimerManager().ClearTimer(SpawnNPCTimer);
+}
+
+void ATwinStickSpawner::SpawnNPCGroup()
+{
+	// reset the group spawn counter
+	SpawnCount = 0;
+
+	// check if we're still under the max NPC cap
+	if (ATwinStickGameMode* GM = Cast<ATwinStickGameMode>(GetWorld()->GetAuthGameMode()))
+	{
+		if (GM->CanSpawnNPCs())
+		{
+			SpawnNPC();
+		}
+	}
+}
+
+void ATwinStickSpawner::SpawnNPC()
+{
+	FTransform SpawnTransform;
+
+	// find a random point around the spawner
+	FVector SpawnLoc;
+	if (UNavigationSystemV1::K2_GetRandomReachablePointInRadius(GetWorld(), GetActorLocation(), SpawnLoc, SpawnRadius, NavData))
+	{
+		SpawnTransform.SetLocation(SpawnLoc);
+
+		// spawn the NPC
+		ATwinStickNPC* NPC = GetWorld()->SpawnActor<ATwinStickNPC>(NPCClass, SpawnTransform);
+	}
+
+	// increase the spawn counter
+	++SpawnCount;
+
+	// do we still have enemies left to spawn?
+	if (SpawnCount < SpawnGroupSize)
+	{
+		GetWorld()->GetTimerManager().SetTimer(SpawnNPCTimer, this, &ATwinStickSpawner::SpawnNPC, FMath::RandRange(MinSpawnDelay, MaxSpawnDelay), false);
+	}
+
+}
diff --git a/Templates/TP_TopDown/Source/TP_TopDown/Variant_TwinStick/TwinStickSpawner.h b/Templates/TP_TopDown/Source/TP_TopDown/Variant_TwinStick/TwinStickSpawner.h
new file mode 100644
--- /dev/null
+++ b/Templates/TP_TopDown/Source/TP_TopDown/Variant_TwinStick/TwinStickSpawner.h
@@ -0,0 +1,79 @@
+// Copyright Epic Games, Inc. All Rights Reserved.
+
+#pragma once
+
+#include "CoreMinimal.h"
+#include "GameFramework/Actor.h"
+#include "TwinStickNPC.h"
+#include "TwinStickSpawner.generated.h"
+
+class ARecastNavMesh;
+
+/**
+ *  A simple NPC spawner for a Twin Stick Shooter game
+ */
+UCLASS(abstract)
+class ATwinStickSpawner : public AActor
+{
+	GENERATED_BODY()
+	
+protected:
+
+	/** Type of NPC to spawn */
+	UPROPERTY(EditAnywhere, Category="NPC Spawner")
+	TSubclassOf<ATwinStickNPC> NPCClass;
+	
+	/** Time delay between enemy group spawns */
+	UPROPERTY(EditAnywhere, Category="NPC Spawner", meta=(ClampMin = 0, ClampMax = 20, Units = "s"))
+	float SpawnGroupDelay = 5.0f;
+
+	/** Min time delay between individual NPC spawns */
+	UPROPERTY(EditAnywhere, Category="NPC Spawner", meta=(ClampMin = 0, ClampMax = 2, Units = "s"))
+	float MinSpawnDelay = 0.33f;
+
+	/** Max time delay between individual NPC spawns */
+	UPROPERTY(EditAnywhere, Category="NPC Spawner", meta=(ClampMin = 0, ClampMax = 2, Units = "s"))
+	float MaxSpawnDelay = 0.66f;
+
+	/** Radius around the spawner where it can spawn NPCs */
+	UPROPERTY(EditAnywhere, Category="NPC Spawner", meta=(ClampMin = 0, ClampMax = 20, Units = "cm"))
+	float SpawnRadius = 600.0f;
+
+	/** Number of NPCs to spawn per group */
+	UPROPERTY(EditAnywhere, Category="NPC Spawner", meta=(ClampMin = 0, ClampMax = 10))
+	int32 SpawnGroupSize = 3;
+	
+	/** Number of NPCs spawned in the current group */
+	int32 SpawnCount = 0;
+
+	/** NPC group spawn timer */
+	FTimerHandle SpawnGroupTimer;
+
+	/** NPC spawn timer */
+	FTimerHandle SpawnNPCTimer;
+
+	/** Pointer to the recast nav mesh actor, used to provide NPC spawn locations */
+	TObjectPtr<ARecastNavMesh> NavData;
+
+public:	
+
+	/** Constructor */
+	ATwinStickSpawner();
+
+protected:
+
+	/** Gameplay initialization */
+	virtual void BeginPlay() override;
+
+	/** Gameplay cleanup */
+	virtual void EndPlay(EEndPlayReason::Type EndPlayReason) override;
+
+protected:
+
+	/** Spawns a new NPC group */
+	void SpawnNPCGroup();
+
+	/** Spawns an individual NPC */
+	void SpawnNPC();
+
+};
diff --git a/Templates/TP_TopDown/Source/TP_TopDown/Variant_TwinStick/TwinStickUI.cpp b/Templates/TP_TopDown/Source/TP_TopDown/Variant_TwinStick/TwinStickUI.cpp
new file mode 100644
--- /dev/null
+++ b/Templates/TP_TopDown/Source/TP_TopDown/Variant_TwinStick/TwinStickUI.cpp
@@ -0,0 +1,5 @@
+// Copyright Epic Games, Inc. All Rights Reserved.
+
+
+#include "TwinStickUI.h"
+
diff --git a/Templates/TP_TopDown/Source/TP_TopDown/Variant_TwinStick/TwinStickUI.h b/Templates/TP_TopDown/Source/TP_TopDown/Variant_TwinStick/TwinStickUI.h
new file mode 100644
--- /dev/null
+++ b/Templates/TP_TopDown/Source/TP_TopDown/Variant_TwinStick/TwinStickUI.h
@@ -0,0 +1,31 @@
+// Copyright Epic Games, Inc. All Rights Reserved.
+
+#pragma once
+
+#include "CoreMinimal.h"
+#include "Blueprint/UserWidget.h"
+#include "TwinStickUI.generated.h"
+
+/**
+ *  A simple Twin Stick Shooter UI widget
+ *  Provides a blueprint interface to expose score values to the UI
+ */
+UCLASS(abstract)
+class UTwinStickUI : public UUserWidget
+{
+	GENERATED_BODY()
+	
+public:
+
+	/** Blueprint handler to update the items counter */
+	UFUNCTION(BlueprintImplementableEvent, Category="Score")
+	void UpdateItems(int32 Score);
+
+	/** Blueprint handler to update the score sub-widgets */
+	UFUNCTION(BlueprintImplementableEvent, Category="Score")
+	void UpdateScore(int32 Score);
+
+	/** Blueprint handler to update the combo sub-widgets */
+	UFUNCTION(BlueprintImplementableEvent, Category="Score")
+	void UpdateCombo(int32 Combo);
+};
diff --git a/Templates/TP_TopDown/TP_TopDown.uproject b/Templates/TP_TopDown/TP_TopDown.uproject
--- a/Templates/TP_TopDown/TP_TopDown.uproject
+++ b/Templates/TP_TopDown/TP_TopDown.uproject
@@ -17,6 +17,14 @@
 			"TargetAllowList": [
 				"Editor"
 			]
+		},
+		{
+			"Name": "StateTree",
+			"Enabled": true
+		},
+		{
+			"Name": "GameplayStateTree",
+			"Enabled": true
 		}
 	]	
 }
\ No newline at end of file
