Templates - TP_TopDown
TwinStick variant - Exposed manual input function nodes to support widget touch interface.
Fixed broken AnimBP reference in TwinStickCharacter.
#rnx

diff --git a/Templates/TP_TopDown/Source/TP_TopDown/Variant_TwinStick/TwinStickCharacter.cpp b/Templates/TP_TopDown/Source/TP_TopDown/Variant_TwinStick/TwinStickCharacter.cpp
--- a/Templates/TP_TopDown/Source/TP_TopDown/Variant_TwinStick/TwinStickCharacter.cpp
+++ b/Templates/TP_TopDown/Source/TP_TopDown/Variant_TwinStick/TwinStickCharacter.cpp
@@ -121,17 +121,10 @@ void ATwinStickCharacter::SetupPlayerInputComponent(UInputComponent* PlayerInput
 void ATwinStickCharacter::Move(const FInputActionValue& Value)
 {
 	// save the input vector
-	LastMoveInput = Value.Get<FVector2D>();
-
-	// calculate the forward component of the input
-	FRotator FlatRot = GetControlRotation();
-	FlatRot.Pitch = 0.0f;
-
-	// apply the forward input
-	AddMovementInput(FlatRot.RotateVector(FVector::ForwardVector), LastMoveInput.X);
+	FVector2D InputVector = Value.Get<FVector2D>();
 
-	// apply the right input
-	AddMovementInput(FlatRot.RotateVector(FVector::RightVector), LastMoveInput.Y);
+	// route the input
+	DoMove(InputVector.X, InputVector.Y);
 }
 
 void ATwinStickCharacter::StickAim(const FInputActionValue& Value)
@@ -139,17 +132,8 @@ void ATwinStickCharacter::StickAim(const FInputActionValue& Value)
 	// get the input vector
 	FVector2D InputVector = Value.Get<FVector2D>();
 
-	// calculate the aim angle from the inputs
-	AimAngle = FMath::RadiansToDegrees(FMath::Atan2(InputVector.Y, -InputVector.X));
-
-	// lower the mouse controls flag
-	bUsingMouse = false;
-
-	// hide the mouse cursor
-	if (PlayerController)
-	{
-		PlayerController->SetShowMouseCursor(false);
-	}
+	// route the input
+	DoAim(InputVector.X, InputVector.Y);
 }
 
 void ATwinStickCharacter::MouseAim(const FInputActionValue& Value)
@@ -165,6 +149,56 @@ void ATwinStickCharacter::MouseAim(const FInputActionValue& Value)
 }
 
 void ATwinStickCharacter::Dash(const FInputActionValue& Value)
+{
+	// route the input
+	DoDash();
+}
+
+void ATwinStickCharacter::Shoot(const FInputActionValue& Value)
+{
+	// route the input
+	DoShoot();
+}
+
+void ATwinStickCharacter::AoEAttack(const FInputActionValue& Value)
+{
+	// route the input
+	DoAoEAttack();
+}
+
+void ATwinStickCharacter::DoMove(float AxisX, float AxisY)
+{
+	// save the input
+	LastMoveInput.X = AxisX;
+	LastMoveInput.Y = AxisY;
+
+	// calculate the forward component of the input
+	FRotator FlatRot = GetControlRotation();
+	FlatRot.Pitch = 0.0f;
+
+	// apply the forward input
+	AddMovementInput(FlatRot.RotateVector(FVector::ForwardVector), AxisX);
+
+	// apply the right input
+	AddMovementInput(FlatRot.RotateVector(FVector::RightVector), AxisY);
+}
+
+void ATwinStickCharacter::DoAim(float AxisX, float AxisY)
+{
+	// calculate the aim angle from the inputs
+	AimAngle = FMath::RadiansToDegrees(FMath::Atan2(AxisY, -AxisX));
+
+	// lower the mouse controls flag
+	bUsingMouse = false;
+
+	// hide the mouse cursor
+	if (PlayerController)
+	{
+		PlayerController->SetShowMouseCursor(false);
+	}
+}
+
+void ATwinStickCharacter::DoDash()
 {
 	// calculate the launch impulse vector based on the last move input
 	FVector LaunchDir = FVector::ZeroVector;
@@ -176,7 +210,7 @@ void ATwinStickCharacter::Dash(const FInputActionValue& Value)
 	LaunchCharacter(LaunchDir * DashImpulse, true, true);
 }
 
-void ATwinStickCharacter::Shoot(const FInputActionValue& Value)
+void ATwinStickCharacter::DoShoot()
 {
 	// get the actor transform
 	FTransform ProjectileTransform = GetActorTransform();
@@ -188,7 +222,7 @@ void ATwinStickCharacter::Shoot(const FInputActionValue& Value)
 	ATwinStickProjectile* Projectile = GetWorld()->SpawnActor<ATwinStickProjectile>(ProjectileClass, ProjectileTransform);
 }
 
-void ATwinStickCharacter::AoEAttack(const FInputActionValue& Value)
+void ATwinStickCharacter::DoAoEAttack()
 {
 	// do we have enough items to do an AoE attack?
 	if (Items > 0)
diff --git a/Templates/TP_TopDown/Source/TP_TopDown/Variant_TwinStick/TwinStickCharacter.h b/Templates/TP_TopDown/Source/TP_TopDown/Variant_TwinStick/TwinStickCharacter.h
--- a/Templates/TP_TopDown/Source/TP_TopDown/Variant_TwinStick/TwinStickCharacter.h
+++ b/Templates/TP_TopDown/Source/TP_TopDown/Variant_TwinStick/TwinStickCharacter.h
@@ -151,6 +151,28 @@ protected:
 	/** Performs an AoE Attack */
 	void AoEAttack(const FInputActionValue& Value);
 
+public:
+
+	/** Handles move inputs from both input actions and touch interface */
+	UFUNCTION(BlueprintCallable, Category="Input")
+	void DoMove(float AxisX, float AxisY);
+
+	/** Handles aim inputs from both input actions and touch interface */
+	UFUNCTION(BlueprintCallable, Category="Input")
+	void DoAim(float AxisX, float AxisY);
+
+	/** Handles dash inputs from both input actions and touch interface */
+	UFUNCTION(BlueprintCallable, Category="Input")
+	void DoDash();
+
+	/** Handles shoot inputs from both input actions and touch interface */
+	UFUNCTION(BlueprintCallable, Category="Input")
+	void DoShoot();
+
+	/** Handles aoe attack inputs from both input actions and touch interface */
+	UFUNCTION(BlueprintCallable, Category="Input")
+	void DoAoEAttack();
+
 public:
 
 	/** Applies collision impact to the player */
