Templates
- Reorganized template variant code
- Cleaned up template properties
- Added units and clamp limits to template properties
- Consolidated log categories in the project module header

#rnx

diff --git a/Templates/TP_FirstPerson/Source/TP_FirstPerson/TP_FirstPerson.Build.cs b/Templates/TP_FirstPerson/Source/TP_FirstPerson/TP_FirstPerson.Build.cs
--- a/Templates/TP_FirstPerson/Source/TP_FirstPerson/TP_FirstPerson.Build.cs
+++ b/Templates/TP_FirstPerson/Source/TP_FirstPerson/TP_FirstPerson.Build.cs
@@ -25,8 +25,11 @@ public class TP_FirstPerson : ModuleRules
 		PublicIncludePaths.AddRange(new string[] {
 			"TP_FirstPerson",
 			"TP_FirstPerson/Variant_Horror",
+			"TP_FirstPerson/Variant_Horror/UI",
 			"TP_FirstPerson/Variant_Shooter",
-			"TP_FirstPerson/Variant_Shooter/AI"
+			"TP_FirstPerson/Variant_Shooter/AI",
+			"TP_FirstPerson/Variant_Shooter/UI",
+			"TP_FirstPerson/Variant_Shooter/Weapons"
 		});
 
 		// Uncomment if you are using Slate UI
diff --git a/Templates/TP_FirstPerson/Source/TP_FirstPerson/TP_FirstPerson.cpp b/Templates/TP_FirstPerson/Source/TP_FirstPerson/TP_FirstPerson.cpp
--- a/Templates/TP_FirstPerson/Source/TP_FirstPerson/TP_FirstPerson.cpp
+++ b/Templates/TP_FirstPerson/Source/TP_FirstPerson/TP_FirstPerson.cpp
@@ -4,4 +4,5 @@
 #include "Modules/ModuleManager.h"
 
 IMPLEMENT_PRIMARY_GAME_MODULE( FDefaultGameModuleImpl, TP_FirstPerson, "TP_FirstPerson" );
- 
\ No newline at end of file
+
+DEFINE_LOG_CATEGORY(LogTP_FirstPerson)
\ No newline at end of file
diff --git a/Templates/TP_FirstPerson/Source/TP_FirstPerson/TP_FirstPerson.h b/Templates/TP_FirstPerson/Source/TP_FirstPerson/TP_FirstPerson.h
--- a/Templates/TP_FirstPerson/Source/TP_FirstPerson/TP_FirstPerson.h
+++ b/Templates/TP_FirstPerson/Source/TP_FirstPerson/TP_FirstPerson.h
@@ -3,3 +3,6 @@
 #pragma once
 
 #include "CoreMinimal.h"
+
+/** Main log category used across the project */
+DECLARE_LOG_CATEGORY_EXTERN(LogTP_FirstPerson, Log, All);
\ No newline at end of file
diff --git a/Templates/TP_FirstPerson/Source/TP_FirstPerson/TP_FirstPersonCharacter.cpp b/Templates/TP_FirstPerson/Source/TP_FirstPerson/TP_FirstPersonCharacter.cpp
--- a/Templates/TP_FirstPerson/Source/TP_FirstPerson/TP_FirstPersonCharacter.cpp
+++ b/Templates/TP_FirstPerson/Source/TP_FirstPerson/TP_FirstPersonCharacter.cpp
@@ -8,11 +8,7 @@
 #include "EnhancedInputComponent.h"
 #include "InputActionValue.h"
 #include "GameFramework/CharacterMovementComponent.h"
-
-DEFINE_LOG_CATEGORY(LogTemplateCharacter);
-
-//////////////////////////////////////////////////////////////////////////
-// ATP_FirstPersonCharacter
+#include "TP_FirstPerson.h"
 
 ATP_FirstPersonCharacter::ATP_FirstPersonCharacter()
 {
@@ -66,7 +62,7 @@ void ATP_FirstPersonCharacter::SetupPlayerInputComponent(UInputComponent* Player
 	}
 	else
 	{
-		UE_LOG(LogTemplateCharacter, Error, TEXT("'%s' Failed to find an Enhanced Input Component! This template is built to use the Enhanced Input system. If you intend to use the legacy system, then you will need to update this C++ file."), *GetNameSafe(this));
+		UE_LOG(LogTP_FirstPerson, Error, TEXT("'%s' Failed to find an Enhanced Input Component! This template is built to use the Enhanced Input system. If you intend to use the legacy system, then you will need to update this C++ file."), *GetNameSafe(this));
 	}
 }
 
diff --git a/Templates/TP_FirstPerson/Source/TP_FirstPerson/TP_FirstPersonCharacter.h b/Templates/TP_FirstPerson/Source/TP_FirstPerson/TP_FirstPersonCharacter.h
--- a/Templates/TP_FirstPerson/Source/TP_FirstPerson/TP_FirstPersonCharacter.h
+++ b/Templates/TP_FirstPerson/Source/TP_FirstPerson/TP_FirstPersonCharacter.h
@@ -34,19 +34,19 @@ class ATP_FirstPersonCharacter : public ACharacter
 protected:
 
 	/** Jump Input Action */
-	UPROPERTY(EditAnywhere, BlueprintReadOnly, Category ="Input")
+	UPROPERTY(EditAnywhere, Category ="Input")
 	UInputAction* JumpAction;
 
 	/** Move Input Action */
-	UPROPERTY(EditAnywhere, BlueprintReadOnly, Category ="Input")
+	UPROPERTY(EditAnywhere, Category ="Input")
 	UInputAction* MoveAction;
 
 	/** Look Input Action */
-	UPROPERTY(EditAnywhere, BlueprintReadOnly, Category ="Input")
+	UPROPERTY(EditAnywhere, Category ="Input")
 	class UInputAction* LookAction;
 
 	/** Mouse Look Input Action */
-	UPROPERTY(EditAnywhere, BlueprintReadOnly, Category ="Input")
+	UPROPERTY(EditAnywhere, Category ="Input")
 	class UInputAction* MouseLookAction;
 	
 public:
diff --git a/Templates/TP_FirstPerson/Source/TP_FirstPerson/TP_FirstPersonPlayerController.h b/Templates/TP_FirstPerson/Source/TP_FirstPerson/TP_FirstPersonPlayerController.h
--- a/Templates/TP_FirstPerson/Source/TP_FirstPerson/TP_FirstPersonPlayerController.h
+++ b/Templates/TP_FirstPerson/Source/TP_FirstPerson/TP_FirstPersonPlayerController.h
@@ -26,7 +26,7 @@ public:
 protected:
 
 	/** Input Mapping Contexts */
-	UPROPERTY(EditAnywhere, BlueprintReadOnly, Category ="Input", meta = (AllowPrivateAccess = "true"))
+	UPROPERTY(EditAnywhere, Category="Input")
 	TArray<UInputMappingContext*> DefaultMappingContexts;
 
 	/** Input mapping context setup */
diff --git a/Templates/TP_FirstPerson/Source/TP_FirstPerson/Variant_Horror/HorrorCharacter.h b/Templates/TP_FirstPerson/Source/TP_FirstPerson/Variant_Horror/HorrorCharacter.h
--- a/Templates/TP_FirstPerson/Source/TP_FirstPerson/Variant_Horror/HorrorCharacter.h
+++ b/Templates/TP_FirstPerson/Source/TP_FirstPerson/Variant_Horror/HorrorCharacter.h
@@ -28,7 +28,7 @@ class TP_FIRSTPERSON_API AHorrorCharacter : public ATP_FirstPersonCharacter
 protected:
 
 	/** Fire weapon input action */
-	UPROPERTY(EditAnywhere, BlueprintReadOnly, Category ="Input")
+	UPROPERTY(EditAnywhere, Category ="Input")
 	UInputAction* SprintAction;
 
 	/** If true, we're sprinting */
@@ -42,26 +42,26 @@ protected:
 	float WalkSpeed = 250.0f;
 
 	/** Time interval for sprinting stamina ticks */
-	UPROPERTY(EditAnywhere, Category="Sprint")
+	UPROPERTY(EditAnywhere, Category="Sprint", meta = (ClampMin = 0, ClampMax = 1, Units = "s"))
 	float SprintFixedTickTime = 0.03333f;
 
 	/** Sprint stamina amount. Maxes at SprintTime */
 	float SprintMeter = 0.0f;
 
 	/** How long we can sprint for, in seconds */
-	UPROPERTY(EditAnywhere, Category="Sprint", meta=(ClampMin = 0, ClampMax = 10, Units = "s"))
+	UPROPERTY(EditAnywhere, Category="Sprint", meta = (ClampMin = 0, ClampMax = 10, Units = "s"))
 	float SprintTime = 3.0f;
 
 	/** Walk speed while sprinting */
-	UPROPERTY(EditAnywhere, Category="Sprint", meta=(ClampMin = 0, ClampMax = 10, Units = "cm/s"))
+	UPROPERTY(EditAnywhere, Category="Sprint", meta = (ClampMin = 0, ClampMax = 10, Units = "cm/s"))
 	float SprintSpeed = 600.0f;
 
 	/** Walk speed while recovering stamina */
-	UPROPERTY(EditAnywhere, Category="Recovery")
+	UPROPERTY(EditAnywhere, Category="Recovery", meta = (ClampMin = 0, ClampMax = 10, Units = "cm/s"))
 	float RecoveringWalkSpeed = 150.0f;
 
 	/** Time it takes for the sprint meter to recover */
-	UPROPERTY(EditAnywhere, Category="Recovery", meta=(ClampMin = 0, ClampMax = 10, Units = "s"))
+	UPROPERTY(EditAnywhere, Category="Recovery", meta = (ClampMin = 0, ClampMax = 10, Units = "s"))
 	float RecoveryTime = 0.0f;
 
 	/** Sprint tick timer */
diff --git a/Templates/TP_FirstPerson/Source/TP_FirstPerson/Variant_Horror/HorrorPlayerController.h b/Templates/TP_FirstPerson/Source/TP_FirstPerson/Variant_Horror/HorrorPlayerController.h
--- a/Templates/TP_FirstPerson/Source/TP_FirstPerson/Variant_Horror/HorrorPlayerController.h
+++ b/Templates/TP_FirstPerson/Source/TP_FirstPerson/Variant_Horror/HorrorPlayerController.h
@@ -39,7 +39,7 @@ protected:
 	virtual void OnPossess(APawn* aPawn) override;
 
 	/** Input Mapping Contexts */
-	UPROPERTY(EditAnywhere, BlueprintReadOnly, Category ="Input", meta = (AllowPrivateAccess = "true"))
+	UPROPERTY(EditAnywhere, Category ="Input")
 	TArray<UInputMappingContext*> DefaultMappingContexts;
 
 	/** Input mapping context setup */
diff --git a/Templates/TP_FirstPerson/Source/TP_FirstPerson/Variant_Shooter/ShooterCharacter.h b/Templates/TP_FirstPerson/Source/TP_FirstPerson/Variant_Shooter/ShooterCharacter.h
--- a/Templates/TP_FirstPerson/Source/TP_FirstPerson/Variant_Shooter/ShooterCharacter.h
+++ b/Templates/TP_FirstPerson/Source/TP_FirstPerson/Variant_Shooter/ShooterCharacter.h
@@ -31,23 +31,23 @@ class TP_FIRSTPERSON_API AShooterCharacter : public ATP_FirstPersonCharacter, pu
 protected:
 
 	/** Fire weapon input action */
-	UPROPERTY(EditAnywhere, BlueprintReadOnly, Category ="Input")
+	UPROPERTY(EditAnywhere, Category ="Input")
 	UInputAction* FireAction;
 
 	/** Switch weapon input action */
-	UPROPERTY(EditAnywhere, BlueprintReadOnly, Category ="Input")
+	UPROPERTY(EditAnywhere, Category ="Input")
 	UInputAction* SwitchWeaponAction;
 
 	/** Name of the first person mesh weapon socket */
-	UPROPERTY(EditAnywhere, BlueprintReadOnly, Category ="Weapons")
+	UPROPERTY(EditAnywhere, Category ="Weapons")
 	FName FirstPersonWeaponSocket = FName("HandGrip_R");
 
 	/** Name of the third person mesh weapon socket */
-	UPROPERTY(EditAnywhere, BlueprintReadOnly, Category ="Weapons")
+	UPROPERTY(EditAnywhere, Category ="Weapons")
 	FName ThirdPersonWeaponSocket = FName("HandGrip_R");
 
 	/** Max distance to use for aim traces */
-	UPROPERTY(EditAnywhere, BlueprintReadOnly, Category ="Aim")
+	UPROPERTY(EditAnywhere, Category ="Aim", meta = (ClampMin = 0, ClampMax = 100000, Units = "cm"))
 	float MaxAimDistance = 10000.0f;
 
 	/** Current HP remaining to this character */
diff --git a/Templates/TP_FirstPerson/Source/TP_FirstPerson/Variant_Shooter/UI/ShooterBulletCounterUI.cpp b/Templates/TP_FirstPerson/Source/TP_FirstPerson/Variant_Shooter/UI/ShooterBulletCounterUI.cpp
new file mode 100644
--- /dev/null
+++ b/Templates/TP_FirstPerson/Source/TP_FirstPerson/Variant_Shooter/UI/ShooterBulletCounterUI.cpp
@@ -0,0 +1,5 @@
+// Copyright Epic Games, Inc. All Rights Reserved.
+
+
+#include "ShooterBulletCounterUI.h"
+
diff --git a/Templates/TP_FirstPerson/Source/TP_FirstPerson/Variant_Shooter/UI/ShooterUI.cpp b/Templates/TP_FirstPerson/Source/TP_FirstPerson/Variant_Shooter/UI/ShooterUI.cpp
new file mode 100644
--- /dev/null
+++ b/Templates/TP_FirstPerson/Source/TP_FirstPerson/Variant_Shooter/UI/ShooterUI.cpp
@@ -0,0 +1,5 @@
+// Copyright Epic Games, Inc. All Rights Reserved.
+
+
+#include "ShooterUI.h"
+
diff --git a/Templates/TP_ThirdPerson/Source/TP_ThirdPerson/TP_ThirdPerson.Build.cs b/Templates/TP_ThirdPerson/Source/TP_ThirdPerson/TP_ThirdPerson.Build.cs
--- a/Templates/TP_ThirdPerson/Source/TP_ThirdPerson/TP_ThirdPerson.Build.cs
+++ b/Templates/TP_ThirdPerson/Source/TP_ThirdPerson/TP_ThirdPerson.Build.cs
@@ -25,11 +25,18 @@ public class TP_ThirdPerson : ModuleRules
 		PublicIncludePaths.AddRange(new string[] {
 			"TP_ThirdPerson",
 			"TP_ThirdPerson/Variant_Platforming",
+			"TP_ThirdPerson/Variant_Platforming/Animation",
 			"TP_ThirdPerson/Variant_Combat",
 			"TP_ThirdPerson/Variant_Combat/AI",
+			"TP_ThirdPerson/Variant_Combat/Animation",
+			"TP_ThirdPerson/Variant_Combat/Gameplay",
+			"TP_ThirdPerson/Variant_Combat/Interfaces",
+			"TP_ThirdPerson/Variant_Combat/UI",
 			"TP_ThirdPerson/Variant_SideScrolling",
+			"TP_ThirdPerson/Variant_SideScrolling/AI",
 			"TP_ThirdPerson/Variant_SideScrolling/Gameplay",
-			"TP_ThirdPerson/Variant_SideScrolling/AI"
+			"TP_ThirdPerson/Variant_SideScrolling/Interfaces",
+			"TP_ThirdPerson/Variant_SideScrolling/UI"
 		});
 
 		// Uncomment if you are using Slate UI
diff --git a/Templates/TP_ThirdPerson/Source/TP_ThirdPerson/TP_ThirdPerson.cpp b/Templates/TP_ThirdPerson/Source/TP_ThirdPerson/TP_ThirdPerson.cpp
--- a/Templates/TP_ThirdPerson/Source/TP_ThirdPerson/TP_ThirdPerson.cpp
+++ b/Templates/TP_ThirdPerson/Source/TP_ThirdPerson/TP_ThirdPerson.cpp
@@ -4,4 +4,5 @@
 #include "Modules/ModuleManager.h"
 
 IMPLEMENT_PRIMARY_GAME_MODULE( FDefaultGameModuleImpl, TP_ThirdPerson, "TP_ThirdPerson" );
- 
\ No newline at end of file
+
+DEFINE_LOG_CATEGORY(LogTP_ThirdPerson)
\ No newline at end of file
diff --git a/Templates/TP_ThirdPerson/Source/TP_ThirdPerson/TP_ThirdPerson.h b/Templates/TP_ThirdPerson/Source/TP_ThirdPerson/TP_ThirdPerson.h
--- a/Templates/TP_ThirdPerson/Source/TP_ThirdPerson/TP_ThirdPerson.h
+++ b/Templates/TP_ThirdPerson/Source/TP_ThirdPerson/TP_ThirdPerson.h
@@ -3,3 +3,6 @@
 #pragma once
 
 #include "CoreMinimal.h"
+
+/** Main log category used across the project */
+DECLARE_LOG_CATEGORY_EXTERN(LogTP_ThirdPerson, Log, All);
\ No newline at end of file
diff --git a/Templates/TP_ThirdPerson/Source/TP_ThirdPerson/TP_ThirdPersonCharacter.cpp b/Templates/TP_ThirdPerson/Source/TP_ThirdPerson/TP_ThirdPersonCharacter.cpp
--- a/Templates/TP_ThirdPerson/Source/TP_ThirdPerson/TP_ThirdPersonCharacter.cpp
+++ b/Templates/TP_ThirdPerson/Source/TP_ThirdPerson/TP_ThirdPersonCharacter.cpp
@@ -10,8 +10,7 @@
 #include "EnhancedInputComponent.h"
 #include "EnhancedInputSubsystems.h"
 #include "InputActionValue.h"
-
-DEFINE_LOG_CATEGORY(LogTemplateCharacter);
+#include "TP_ThirdPerson.h"
 
 ATP_ThirdPersonCharacter::ATP_ThirdPersonCharacter()
 {
@@ -69,7 +68,7 @@ void ATP_ThirdPersonCharacter::SetupPlayerInputComponent(UInputComponent* Player
 	}
 	else
 	{
-		UE_LOG(LogTemplateCharacter, Error, TEXT("'%s' Failed to find an Enhanced Input component! This template is built to use the Enhanced Input system. If you intend to use the legacy system, then you will need to update this C++ file."), *GetNameSafe(this));
+		UE_LOG(LogTP_ThirdPerson, Error, TEXT("'%s' Failed to find an Enhanced Input component! This template is built to use the Enhanced Input system. If you intend to use the legacy system, then you will need to update this C++ file."), *GetNameSafe(this));
 	}
 }
 
diff --git a/Templates/TP_ThirdPerson/Source/TP_ThirdPerson/TP_ThirdPersonCharacter.h b/Templates/TP_ThirdPerson/Source/TP_ThirdPerson/TP_ThirdPersonCharacter.h
--- a/Templates/TP_ThirdPerson/Source/TP_ThirdPerson/TP_ThirdPersonCharacter.h
+++ b/Templates/TP_ThirdPerson/Source/TP_ThirdPerson/TP_ThirdPersonCharacter.h
@@ -24,29 +24,29 @@ class ATP_ThirdPersonCharacter : public ACharacter
 	GENERATED_BODY()
 
 	/** Camera boom positioning the camera behind the character */
-	UPROPERTY(VisibleAnywhere, BlueprintReadOnly, Category = Camera, meta = (AllowPrivateAccess = "true"))
+	UPROPERTY(VisibleAnywhere, BlueprintReadOnly, Category="Components", meta = (AllowPrivateAccess = "true"))
 	USpringArmComponent* CameraBoom;
 
 	/** Follow camera */
-	UPROPERTY(VisibleAnywhere, BlueprintReadOnly, Category = Camera, meta = (AllowPrivateAccess = "true"))
+	UPROPERTY(VisibleAnywhere, BlueprintReadOnly, Category="Components", meta = (AllowPrivateAccess = "true"))
 	UCameraComponent* FollowCamera;
 	
 protected:
 
 	/** Jump Input Action */
-	UPROPERTY(EditAnywhere, BlueprintReadOnly, Category="Input")
+	UPROPERTY(EditAnywhere, Category="Input")
 	UInputAction* JumpAction;
 
 	/** Move Input Action */
-	UPROPERTY(EditAnywhere, BlueprintReadOnly, Category="Input")
+	UPROPERTY(EditAnywhere, Category="Input")
 	UInputAction* MoveAction;
 
 	/** Look Input Action */
-	UPROPERTY(EditAnywhere, BlueprintReadOnly, Category="Input")
+	UPROPERTY(EditAnywhere, Category="Input")
 	UInputAction* LookAction;
 
 	/** Mouse Look Input Action */
-	UPROPERTY(EditAnywhere, BlueprintReadOnly, Category="Input")
+	UPROPERTY(EditAnywhere, Category="Input")
 	UInputAction* MouseLookAction;
 
 public:
diff --git a/Templates/TP_ThirdPerson/Source/TP_ThirdPerson/TP_ThirdPersonPlayerController.h b/Templates/TP_ThirdPerson/Source/TP_ThirdPerson/TP_ThirdPersonPlayerController.h
--- a/Templates/TP_ThirdPerson/Source/TP_ThirdPerson/TP_ThirdPersonPlayerController.h
+++ b/Templates/TP_ThirdPerson/Source/TP_ThirdPerson/TP_ThirdPersonPlayerController.h
@@ -20,7 +20,7 @@ class ATP_ThirdPersonPlayerController : public APlayerController
 protected:
 
 	/** Input Mapping Contexts */
-	UPROPERTY(EditAnywhere, BlueprintReadOnly, Category ="Input", meta = (AllowPrivateAccess = "true"))
+	UPROPERTY(EditAnywhere, Category ="Input")
 	TArray<UInputMappingContext*> DefaultMappingContexts;
 
 	/** Input mapping context setup */
diff --git a/Templates/TP_ThirdPerson/Source/TP_ThirdPerson/Variant_Combat/AI/CombatAIController.h b/Templates/TP_ThirdPerson/Source/TP_ThirdPerson/Variant_Combat/AI/CombatAIController.h
--- a/Templates/TP_ThirdPerson/Source/TP_ThirdPerson/Variant_Combat/AI/CombatAIController.h
+++ b/Templates/TP_ThirdPerson/Source/TP_ThirdPerson/Variant_Combat/AI/CombatAIController.h
@@ -17,7 +17,7 @@ class ACombatAIController : public AAIController
 	GENERATED_BODY()
 
 	/** StateTree Component */
-	UPROPERTY(VisibleAnywhere, BlueprintReadOnly, Category = "AI", meta = (AllowPrivateAccess = "true"))
+	UPROPERTY(VisibleAnywhere, BlueprintReadOnly, Category = "Components", meta = (AllowPrivateAccess = "true"))
 	UStateTreeAIComponent* StateTreeAI;
 
 public:
diff --git a/Templates/TP_ThirdPerson/Source/TP_ThirdPerson/Variant_Combat/AI/CombatEnemy.h b/Templates/TP_ThirdPerson/Source/TP_ThirdPerson/Variant_Combat/AI/CombatEnemy.h
--- a/Templates/TP_ThirdPerson/Source/TP_ThirdPerson/Variant_Combat/AI/CombatEnemy.h
+++ b/Templates/TP_ThirdPerson/Source/TP_ThirdPerson/Variant_Combat/AI/CombatEnemy.h
@@ -33,7 +33,7 @@ class ACombatEnemy : public ACharacter, public ICombatAttacker, public ICombatDa
 	GENERATED_BODY()
 
 	/** Life bar widget component */
-	UPROPERTY(VisibleAnywhere, BlueprintReadOnly, Category = UI, meta = (AllowPrivateAccess = "true"))
+	UPROPERTY(VisibleAnywhere, BlueprintReadOnly, Category="Components", meta = (AllowPrivateAccess = "true"))
 	UWidgetComponent* LifeBar;
 
 public:
diff --git a/Templates/TP_ThirdPerson/Source/TP_ThirdPerson/Variant_Combat/AI/CombatEnemySpawner.h b/Templates/TP_ThirdPerson/Source/TP_ThirdPerson/Variant_Combat/AI/CombatEnemySpawner.h
--- a/Templates/TP_ThirdPerson/Source/TP_ThirdPerson/Variant_Combat/AI/CombatEnemySpawner.h
+++ b/Templates/TP_ThirdPerson/Source/TP_ThirdPerson/Variant_Combat/AI/CombatEnemySpawner.h
@@ -22,10 +22,10 @@ class ACombatEnemySpawner : public AActor, public ICombatActivatable
 {
 	GENERATED_BODY()
 	
-	UPROPERTY(VisibleAnywhere, BlueprintReadOnly, Category = Components, meta = (AllowPrivateAccess = "true"))
+	UPROPERTY(VisibleAnywhere, BlueprintReadOnly, Category="Components", meta = (AllowPrivateAccess = "true"))
 	UCapsuleComponent* SpawnCapsule;
 
-	UPROPERTY(VisibleAnywhere, BlueprintReadOnly, Category = Components, meta = (AllowPrivateAccess = "true"))
+	UPROPERTY(VisibleAnywhere, BlueprintReadOnly, Category = "Components", meta = (AllowPrivateAccess = "true"))
 	UArrowComponent* SpawnDirection;
 
 protected:
diff --git a/Templates/TP_ThirdPerson/Source/TP_ThirdPerson/Variant_Combat/CombatCharacter.cpp b/Templates/TP_ThirdPerson/Source/TP_ThirdPerson/Variant_Combat/CombatCharacter.cpp
--- a/Templates/TP_ThirdPerson/Source/TP_ThirdPerson/Variant_Combat/CombatCharacter.cpp
+++ b/Templates/TP_ThirdPerson/Source/TP_ThirdPerson/Variant_Combat/CombatCharacter.cpp
@@ -16,8 +16,6 @@
 #include "Engine/LocalPlayer.h"
 #include "CombatPlayerController.h"
 
-DEFINE_LOG_CATEGORY(LogCombatCharacter);
-
 ACombatCharacter::ACombatCharacter()
 {
 	PrimaryActorTick.bCanEverTick = true;
diff --git a/Templates/TP_ThirdPerson/Source/TP_ThirdPerson/Variant_Combat/CombatCharacter.h b/Templates/TP_ThirdPerson/Source/TP_ThirdPerson/Variant_Combat/CombatCharacter.h
--- a/Templates/TP_ThirdPerson/Source/TP_ThirdPerson/Variant_Combat/CombatCharacter.h
+++ b/Templates/TP_ThirdPerson/Source/TP_ThirdPerson/Variant_Combat/CombatCharacter.h
@@ -32,15 +32,15 @@ class ACombatCharacter : public ACharacter, public ICombatAttacker, public IComb
 	GENERATED_BODY()
 
 	/** Camera boom positioning the camera behind the character */
-	UPROPERTY(VisibleAnywhere, BlueprintReadOnly, Category = Camera, meta = (AllowPrivateAccess = "true"))
+	UPROPERTY(VisibleAnywhere, BlueprintReadOnly, Category="Components", meta = (AllowPrivateAccess = "true"))
 	USpringArmComponent* CameraBoom;
 
 	/** Follow camera */
-	UPROPERTY(VisibleAnywhere, BlueprintReadOnly, Category = Camera, meta = (AllowPrivateAccess = "true"))
+	UPROPERTY(VisibleAnywhere, BlueprintReadOnly, Category="Components", meta = (AllowPrivateAccess = "true"))
 	UCameraComponent* FollowCamera;
 
 	/** Life bar widget component */
-	UPROPERTY(VisibleAnywhere, BlueprintReadOnly, Category = UI, meta = (AllowPrivateAccess = "true"))
+	UPROPERTY(VisibleAnywhere, BlueprintReadOnly, Category="Components", meta = (AllowPrivateAccess = "true"))
 	UWidgetComponent* LifeBar;
 	
 protected:
@@ -58,7 +58,7 @@ protected:
 	UInputAction* LookAction;
 
 	/** Mouse Look Input Action */
-	UPROPERTY(EditAnywhere, BlueprintReadOnly, Category="Input")
+	UPROPERTY(EditAnywhere, Category="Input")
 	UInputAction* MouseLookAction;
 
 	/** Combo Attack Input Action */
@@ -90,7 +90,7 @@ protected:
 	TObjectPtr<UCombatLifeBar> LifeBarWidget;
 
 	/** Max amount of time that may elapse for a non-combo attack input to not be considered stale */
-	UPROPERTY(EditAnywhere, Category="Melee Attack", meta = (ClampMin = 0, ClampMax = 5))
+	UPROPERTY(EditAnywhere, Category="Melee Attack", meta = (ClampMin = 0, ClampMax = 5, Units = "s"))
 	float AttackInputCacheTimeTolerance = 1.0f;
 
 	/** Time at which an attack button was last pressed */
@@ -100,7 +100,7 @@ protected:
 	bool bIsAttacking = false;
 
 	/** Distance ahead of the character that melee attack sphere collision traces will extend */
-	UPROPERTY(EditAnywhere, Category="Melee Attack|Trace", meta=(ClampMin=0, ClampMax=500, Units="cm"))
+	UPROPERTY(EditAnywhere, Category="Melee Attack|Trace", meta = (ClampMin = 0, ClampMax = 500, Units="cm"))
 	float MeleeTraceDistance = 75.0f;
 
 	/** Radius of the sphere trace for melee attacks */
@@ -128,7 +128,7 @@ protected:
 	TArray<FName> ComboSectionNames;
 
 	/** Max amount of time that may elapse for a combo attack input to not be considered stale */
-	UPROPERTY(EditAnywhere, Category="Melee Attack|Combo", meta = (ClampMin = 0, ClampMax = 5))
+	UPROPERTY(EditAnywhere, Category="Melee Attack|Combo", meta = (ClampMin = 0, ClampMax = 5, Units = "s"))
 	float ComboInputCacheTimeTolerance = 0.45f;
 
 	/** Index of the current stage of the melee attack combo */
@@ -161,7 +161,7 @@ protected:
 	float DefaultCameraDistance = 100.0f;
 
 	/** Time to wait before respawning the character */
-	UPROPERTY(EditAnywhere, Category="Respawn", meta = (ClampMin = 0, ClampMax = 10))
+	UPROPERTY(EditAnywhere, Category="Respawn", meta = (ClampMin = 0, ClampMax = 10, Units = "s"))
 	float RespawnTime = 3.0f;
 
 	/** Attack montage ended delegate */
diff --git a/Templates/TP_ThirdPerson/Source/TP_ThirdPerson/Variant_Combat/UI/CombatLifeBar.h b/Templates/TP_ThirdPerson/Source/TP_ThirdPerson/Variant_Combat/UI/CombatLifeBar.h
new file mode 100644
--- /dev/null
+++ b/Templates/TP_ThirdPerson/Source/TP_ThirdPerson/Variant_Combat/UI/CombatLifeBar.h
@@ -0,0 +1,26 @@
+// Copyright Epic Games, Inc. All Rights Reserved.
+
+#pragma once
+
+#include "CoreMinimal.h"
+#include "Blueprint/UserWidget.h"
+#include "CombatLifeBar.generated.h"
+
+/**
+ *  A basic life bar user widget.
+ */
+UCLASS(abstract)
+class UCombatLifeBar : public UUserWidget
+{
+	GENERATED_BODY()
+
+public:
+
+	/** Sets the life bar to the provided 0-1 percentage value*/
+	UFUNCTION(BlueprintImplementableEvent, Category="Life Bar")
+	void SetLifePercentage(float Percent);
+
+	// Sets the life bar fill color
+	UFUNCTION(BlueprintImplementableEvent, Category="Life Bar")
+	void SetBarColor(FLinearColor Color);
+};
diff --git a/Templates/TP_ThirdPerson/Source/TP_ThirdPerson/Variant_Platforming/PlatformingCharacter.h b/Templates/TP_ThirdPerson/Source/TP_ThirdPerson/Variant_Platforming/PlatformingCharacter.h
--- a/Templates/TP_ThirdPerson/Source/TP_ThirdPerson/Variant_Platforming/PlatformingCharacter.h
+++ b/Templates/TP_ThirdPerson/Source/TP_ThirdPerson/Variant_Platforming/PlatformingCharacter.h
@@ -28,33 +28,33 @@ class APlatformingCharacter : public ACharacter
 	GENERATED_BODY()
 
 	/** Camera boom positioning the camera behind the character */
-	UPROPERTY(VisibleAnywhere, BlueprintReadOnly, Category = Camera, meta = (AllowPrivateAccess = "true"))
+	UPROPERTY(VisibleAnywhere, BlueprintReadOnly, Category="Components", meta = (AllowPrivateAccess = "true"))
 	USpringArmComponent* CameraBoom;
 
 	/** Follow camera */
-	UPROPERTY(VisibleAnywhere, BlueprintReadOnly, Category = Camera, meta = (AllowPrivateAccess = "true"))
+	UPROPERTY(VisibleAnywhere, BlueprintReadOnly, Category="Components", meta = (AllowPrivateAccess = "true"))
 	UCameraComponent* FollowCamera;
 	
 protected:
 
 	/** Jump Input Action */
-	UPROPERTY(EditAnywhere, BlueprintReadOnly, Category="Input")
+	UPROPERTY(EditAnywhere, Category="Input")
 	UInputAction* JumpAction;
 
 	/** Move Input Action */
-	UPROPERTY(EditAnywhere, BlueprintReadOnly, Category="Input")
+	UPROPERTY(EditAnywhere, Category="Input")
 	UInputAction* MoveAction;
 
 	/** Look Input Action */
-	UPROPERTY(EditAnywhere, BlueprintReadOnly, Category="Input")
+	UPROPERTY(EditAnywhere, Category="Input")
 	UInputAction* LookAction;
 
 	/** Mouse Look Input Action */
-	UPROPERTY(EditAnywhere, BlueprintReadOnly, Category="Input")
+	UPROPERTY(EditAnywhere, Category="Input")
 	UInputAction* MouseLookAction;
 
 	/** Dash Input Action */
-	UPROPERTY(EditAnywhere, BlueprintReadOnly, Category="Input")
+	UPROPERTY(EditAnywhere, Category="Input")
 	UInputAction* DashAction;
 
 public:
@@ -151,23 +151,23 @@ protected:
 	FOnMontageEnded OnDashMontageEnded;
 
 	/** Distance to trace ahead of the character to look for walls to jump from */
-	UPROPERTY(EditAnywhere, Category="Wall Jump")
+	UPROPERTY(EditAnywhere, Category="Wall Jump", meta = (ClampMin = 0, ClampMax = 1000, Units = "cm"))
 	float WallJumpTraceDistance = 50.0f;
 
 	/** Radius of the wall jump sphere trace check */
-	UPROPERTY(EditAnywhere, Category="Wall Jump")
+	UPROPERTY(EditAnywhere, Category="Wall Jump", meta = (ClampMin = 0, ClampMax = 100, Units = "cm"))
 	float WallJumpTraceRadius = 25.0f;
 
 	/** Impulse to apply away from the wall when wall jumping */
-	UPROPERTY(EditAnywhere, Category="Wall Jump")
+	UPROPERTY(EditAnywhere, Category="Wall Jump", meta = (ClampMin = 0, ClampMax = 10000, Units = "cm/s"))
 	float WallJumpBounceImpulse = 800.0f;
 
 	/** Vertical impulse to apply when wall jumping */
-	UPROPERTY(EditAnywhere, Category="Wall Jump")
+	UPROPERTY(EditAnywhere, Category="Wall Jump", meta = (ClampMin = 0, ClampMax = 10000, Units = "cm/s"))
 	float WallJumpVerticalImpulse = 900.0f;
 
 	/** Time to ignore jump inputs after a wall jump */
-	UPROPERTY(EditAnywhere, Category="Wall Jump")
+	UPROPERTY(EditAnywhere, Category="Wall Jump", meta = (ClampMin = 0, ClampMax = 5, Units = "s"))
 	float DelayBetweenWallJumps = 0.1f;
 
 	/** AnimMontage to use for the Dash action */
diff --git a/Templates/TP_ThirdPerson/Source/TP_ThirdPerson/Variant_SideScrolling/AI/SideScrollingNPC.h b/Templates/TP_ThirdPerson/Source/TP_ThirdPerson/Variant_SideScrolling/AI/SideScrollingNPC.h
--- a/Templates/TP_ThirdPerson/Source/TP_ThirdPerson/Variant_SideScrolling/AI/SideScrollingNPC.h
+++ b/Templates/TP_ThirdPerson/Source/TP_ThirdPerson/Variant_SideScrolling/AI/SideScrollingNPC.h
@@ -20,15 +20,15 @@ class ASideScrollingNPC : public ACharacter, public ISideScrollingInteractable
 protected:
 
 	/** Horizontal impulse to apply to the NPC when it's interacted with */
-	UPROPERTY(EditAnywhere, Category="NPC", meta=(ClampMin=0, Units="cm/s"))
+	UPROPERTY(EditAnywhere, Category="NPC", meta = (ClampMin = 0, ClampMax = 10000, Units="cm/s"))
 	float LaunchImpulse = 500.0f;
 
 	/** Vertical impulse to apply to the NPC when it's interacted with */
-	UPROPERTY(EditAnywhere, Category="NPC", meta=(ClampMin=0, Units="cm/s"))
+	UPROPERTY(EditAnywhere, Category="NPC", meta = (ClampMin = 0, ClampMax = 10000, Units="cm/s"))
 	float LaunchVerticalImpulse = 500.0f;
 
 	/** Time that the NPC remains deactivated after being interacted with */
-	UPROPERTY(EditAnywhere, Category="NPC", meta=(ClampMin=0, Units="s"))
+	UPROPERTY(EditAnywhere, Category="NPC", meta = (ClampMin = 0, ClampMax = 10, Units="s"))
 	float DeactivationTime = 3.0f;
 
 public:
diff --git a/Templates/TP_ThirdPerson/Source/TP_ThirdPerson/Variant_SideScrolling/AI/SideScrollingStateTreeUtility.h b/Templates/TP_ThirdPerson/Source/TP_ThirdPerson/Variant_SideScrolling/AI/SideScrollingStateTreeUtility.h
--- a/Templates/TP_ThirdPerson/Source/TP_ThirdPerson/Variant_SideScrolling/AI/SideScrollingStateTreeUtility.h
+++ b/Templates/TP_ThirdPerson/Source/TP_ThirdPerson/Variant_SideScrolling/AI/SideScrollingStateTreeUtility.h
@@ -18,23 +18,23 @@ struct FStateTreeGetPlayerInstanceData
 	GENERATED_BODY()
 
 	/** NPC owning this task */
-	UPROPERTY(VisibleAnywhere, Category = Context)
+	UPROPERTY(VisibleAnywhere, Category="Context")
 	TObjectPtr<APawn> NPC;
 
 	/** Holds the found player pawn */
-	UPROPERTY(VisibleAnywhere, Category = Context)
+	UPROPERTY(VisibleAnywhere, Category="Context")
 	TObjectPtr<AAIController> Controller;
 
 	/** Holds the found player pawn */
-	UPROPERTY(VisibleAnywhere)
+	UPROPERTY(VisibleAnywhere, Category="Output")
 	TObjectPtr<APawn> TargetPlayer;
 
 	/** Is the pawn close enough to be considered a valid target? */
-	UPROPERTY(VisibleAnywhere)
+	UPROPERTY(VisibleAnywhere, Category="Output")
 	bool bValidTarget = false;
 
 	/** Max distance to be considered a valid target */
-	UPROPERTY(EditAnywhere, Category = Parameter, meta=(ClampMin = 0, Units = "cm"))
+	UPROPERTY(EditAnywhere, Category="Parameter", meta = (ClampMin = 0, ClampMax = 10000, Units = "cm"))
 	float RangeMax = 1000.0f;
 };
 
diff --git a/Templates/TP_ThirdPerson/Source/TP_ThirdPerson/Variant_SideScrolling/Gameplay/SideScrollingMovingPlatform.h b/Templates/TP_ThirdPerson/Source/TP_ThirdPerson/Variant_SideScrolling/Gameplay/SideScrollingMovingPlatform.h
--- a/Templates/TP_ThirdPerson/Source/TP_ThirdPerson/Variant_SideScrolling/Gameplay/SideScrollingMovingPlatform.h
+++ b/Templates/TP_ThirdPerson/Source/TP_ThirdPerson/Variant_SideScrolling/Gameplay/SideScrollingMovingPlatform.h
@@ -31,7 +31,7 @@ protected:
 	FVector PlatformTarget;
 
 	/** Time for the platform to move to the destination */
-	UPROPERTY(EditAnywhere, BlueprintReadOnly, Category="Moving Platform", meta=(ClampMin=0, Units="s"))
+	UPROPERTY(EditAnywhere, BlueprintReadOnly, Category="Moving Platform", meta = (ClampMin = 0, ClampMax = 10, Units="s"))
 	float MoveDuration = 5.0f;
 
 	/** If this is true, the platform will only move once. */
@@ -54,7 +54,7 @@ public:
 protected:
 
 	/** Allows Blueprint code to do the actual platform movement */
-	UFUNCTION(BlueprintImplementableEvent, BlueprintCallable, Category="Moving Platform", meta=(DisplayName="Move to Target"))
+	UFUNCTION(BlueprintImplementableEvent, BlueprintCallable, Category="Moving Platform", meta = (DisplayName="Move to Target"))
 	void BP_MoveToTarget();
 
 };
diff --git a/Templates/TP_ThirdPerson/Source/TP_ThirdPerson/Variant_SideScrolling/SideScrollingCameraManager.h b/Templates/TP_ThirdPerson/Source/TP_ThirdPerson/Variant_SideScrolling/SideScrollingCameraManager.h
--- a/Templates/TP_ThirdPerson/Source/TP_ThirdPerson/Variant_SideScrolling/SideScrollingCameraManager.h
+++ b/Templates/TP_ThirdPerson/Source/TP_ThirdPerson/Variant_SideScrolling/SideScrollingCameraManager.h
@@ -22,19 +22,19 @@ public:
 public:
 
 	/** How close we want to stay to the view target */
-	UPROPERTY(EditAnywhere, BlueprintReadWrite, Category="Side Scrolling Camera", meta=(ClampMin=0, ClampMax=10000, Units="cm"))
+	UPROPERTY(EditAnywhere, Category="Side Scrolling Camera", meta=(ClampMin=0, ClampMax=10000, Units="cm"))
 	float CurrentZoom = 1000.0f;
 
 	/** How far above the target do we want the camera to focus */
-	UPROPERTY(EditAnywhere, BlueprintReadWrite, Category="Side Scrolling Camera", meta=(ClampMin=0, ClampMax=10000, Units="cm"))
+	UPROPERTY(EditAnywhere, Category="Side Scrolling Camera", meta=(ClampMin=0, ClampMax=10000, Units="cm"))
 	float CameraZOffset = 100.0f;
 
 	/** Minimum camera scrolling bounds in world space */
-	UPROPERTY(EditAnywhere, BlueprintReadWrite, Category="Side Scrolling Camera", meta=(ClampMin=0, ClampMax=10000, Units="cm"))
+	UPROPERTY(EditAnywhere, Category="Side Scrolling Camera", meta=(ClampMin=0, ClampMax=10000, Units="cm"))
 	float CameraXMinBounds = -400.0f;
 
 	/** Maximum camera scrolling bounds in world space */
-	UPROPERTY(EditAnywhere, BlueprintReadWrite, Category="Side Scrolling Camera", meta=(ClampMin=0, ClampMax=10000, Units="cm"))
+	UPROPERTY(EditAnywhere, Category="Side Scrolling Camera", meta=(ClampMin=0, ClampMax=10000, Units="cm"))
 	float CameraXMaxBounds = 10000.0f;
 
 protected:
diff --git a/Templates/TP_ThirdPerson/Source/TP_ThirdPerson/Variant_SideScrolling/SideScrollingCharacter.h b/Templates/TP_ThirdPerson/Source/TP_ThirdPerson/Variant_SideScrolling/SideScrollingCharacter.h
--- a/Templates/TP_ThirdPerson/Source/TP_ThirdPerson/Variant_SideScrolling/SideScrollingCharacter.h
+++ b/Templates/TP_ThirdPerson/Source/TP_ThirdPerson/Variant_SideScrolling/SideScrollingCharacter.h
@@ -25,19 +25,19 @@ class ASideScrollingCharacter : public ACharacter
 protected:
 
 	/** Move Input Action */
-	UPROPERTY(EditAnywhere, BlueprintReadOnly, Category="Input")
+	UPROPERTY(EditAnywhere, Category="Input")
 	UInputAction* MoveAction;
 
 	/** Jump Input Action */
-	UPROPERTY(EditAnywhere, BlueprintReadOnly, Category="Input")
+	UPROPERTY(EditAnywhere, Category="Input")
 	UInputAction* JumpAction;
 
 	/** Drop from Platform Action */
-	UPROPERTY(EditAnywhere, BlueprintReadOnly, Category="Input")
+	UPROPERTY(EditAnywhere, Category="Input")
 	UInputAction* DropAction;
 
 	/** Interact Input Action */
-	UPROPERTY(EditAnywhere, BlueprintReadOnly, Category="Input")
+	UPROPERTY(EditAnywhere, Category="Input")
 	UInputAction* InteractAction;
 
 	/** Max distance that interactive objects can be triggered */
diff --git a/Templates/TP_ThirdPerson/Source/TP_ThirdPerson/Variant_SideScrolling/SideScrollingGameMode.h b/Templates/TP_ThirdPerson/Source/TP_ThirdPerson/Variant_SideScrolling/SideScrollingGameMode.h
--- a/Templates/TP_ThirdPerson/Source/TP_ThirdPerson/Variant_SideScrolling/SideScrollingGameMode.h
+++ b/Templates/TP_ThirdPerson/Source/TP_ThirdPerson/Variant_SideScrolling/SideScrollingGameMode.h
@@ -29,7 +29,7 @@ protected:
 	TObjectPtr<USideScrollingUI> UserInterface;
 
 	/** Number of pickups collected by the player */
-	UPROPERTY(BlueprintReadOnly, Category="Picups")
+	UPROPERTY(BlueprintReadOnly, Category="Pickups")
 	int32 PickupsCollected = 0;
 
 protected:
diff --git a/Templates/TP_TopDown/Source/TP_TopDown/TP_TopDown.Build.cs b/Templates/TP_TopDown/Source/TP_TopDown/TP_TopDown.Build.cs
--- a/Templates/TP_TopDown/Source/TP_TopDown/TP_TopDown.Build.cs
+++ b/Templates/TP_TopDown/Source/TP_TopDown/TP_TopDown.Build.cs
@@ -27,8 +27,11 @@ public class TP_TopDown : ModuleRules
 		PublicIncludePaths.AddRange(new string[] {
 			"TP_TopDown",
 			"TP_TopDown/Variant_Strategy",
+			"TP_TopDown/Variant_Strategy/UI",
 			"TP_TopDown/Variant_TwinStick",
-			"TP_TopDown/Variant_TwinStick/AI"
+			"TP_TopDown/Variant_TwinStick/AI",
+			"TP_TopDown/Variant_TwinStick/Gameplay",
+			"TP_TopDown/Variant_TwinStick/UI"
 		});
 
 		// Uncomment if you are using Slate UI
diff --git a/Templates/TP_TopDown/Source/TP_TopDown/TP_TopDown.h b/Templates/TP_TopDown/Source/TP_TopDown/TP_TopDown.h
--- a/Templates/TP_TopDown/Source/TP_TopDown/TP_TopDown.h
+++ b/Templates/TP_TopDown/Source/TP_TopDown/TP_TopDown.h
@@ -4,4 +4,5 @@
 
 #include "CoreMinimal.h"
 
+/** Main log category used across the project */
 DECLARE_LOG_CATEGORY_EXTERN(LogTP_TopDown, Log, All);
diff --git a/Templates/TP_TopDown/Source/TP_TopDown/TP_TopDownCharacter.h b/Templates/TP_TopDown/Source/TP_TopDown/TP_TopDownCharacter.h
--- a/Templates/TP_TopDown/Source/TP_TopDown/TP_TopDownCharacter.h
+++ b/Templates/TP_TopDown/Source/TP_TopDown/TP_TopDownCharacter.h
@@ -17,11 +17,11 @@ class ATP_TopDownCharacter : public ACharacter
 private:
 
 	/** Top down camera */
-	UPROPERTY(VisibleAnywhere, BlueprintReadOnly, Category = Camera, meta = (AllowPrivateAccess = "true"))
+	UPROPERTY(VisibleAnywhere, BlueprintReadOnly, Category="Components", meta = (AllowPrivateAccess = "true"))
 	class UCameraComponent* TopDownCameraComponent;
 
 	/** Camera boom positioning the camera above the character */
-	UPROPERTY(VisibleAnywhere, BlueprintReadOnly, Category = Camera, meta = (AllowPrivateAccess = "true"))
+	UPROPERTY(VisibleAnywhere, BlueprintReadOnly, Category="Components", meta = (AllowPrivateAccess = "true"))
 	class USpringArmComponent* CameraBoom;
 
 public:
diff --git a/Templates/TP_TopDown/Source/TP_TopDown/TP_TopDownPlayerController.cpp b/Templates/TP_TopDown/Source/TP_TopDown/TP_TopDownPlayerController.cpp
--- a/Templates/TP_TopDown/Source/TP_TopDown/TP_TopDownPlayerController.cpp
+++ b/Templates/TP_TopDown/Source/TP_TopDown/TP_TopDownPlayerController.cpp
@@ -11,8 +11,7 @@
 #include "InputActionValue.h"
 #include "EnhancedInputSubsystems.h"
 #include "Engine/LocalPlayer.h"
-
-DEFINE_LOG_CATEGORY(LogTemplateCharacter);
+#include "TP_TopDown.h"
 
 ATP_TopDownPlayerController::ATP_TopDownPlayerController()
 {
@@ -60,7 +59,7 @@ void ATP_TopDownPlayerController::SetupInputComponent()
 	}
 	else
 	{
-		UE_LOG(LogTemplateCharacter, Error, TEXT("'%s' Failed to find an Enhanced Input Component! This template is built to use the Enhanced Input system. If you intend to use the legacy system, then you will need to update this C++ file."), *GetNameSafe(this));
+		UE_LOG(LogTP_TopDown, Error, TEXT("'%s' Failed to find an Enhanced Input Component! This template is built to use the Enhanced Input system. If you intend to use the legacy system, then you will need to update this C++ file."), *GetNameSafe(this));
 	}
 }
 
diff --git a/Templates/TP_TopDown/Source/TP_TopDown/TP_TopDownPlayerController.h b/Templates/TP_TopDown/Source/TP_TopDown/TP_TopDownPlayerController.h
--- a/Templates/TP_TopDown/Source/TP_TopDown/TP_TopDownPlayerController.h
+++ b/Templates/TP_TopDown/Source/TP_TopDown/TP_TopDownPlayerController.h
@@ -25,23 +25,23 @@ class ATP_TopDownPlayerController : public APlayerController
 protected:
 
 	/** Time Threshold to know if it was a short press */
-	UPROPERTY(EditAnywhere, BlueprintReadOnly, Category = Input)
+	UPROPERTY(EditAnywhere, Category="Input")
 	float ShortPressThreshold;
 
 	/** FX Class that we will spawn when clicking */
-	UPROPERTY(EditAnywhere, BlueprintReadOnly, Category = Input)
+	UPROPERTY(EditAnywhere, Category="Input")
 	UNiagaraSystem* FXCursor;
 
 	/** MappingContext */
-	UPROPERTY(EditAnywhere, BlueprintReadOnly, Category=Input, meta=(AllowPrivateAccess = "true"))
+	UPROPERTY(EditAnywhere, Category="Input")
 	UInputMappingContext* DefaultMappingContext;
 	
 	/** Jump Input Action */
-	UPROPERTY(EditAnywhere, BlueprintReadOnly, Category=Input, meta=(AllowPrivateAccess = "true"))
+	UPROPERTY(EditAnywhere, Category="Input")
 	UInputAction* SetDestinationClickAction;
 
 	/** Jump Input Action */
-	UPROPERTY(EditAnywhere, BlueprintReadOnly, Category=Input, meta=(AllowPrivateAccess = "true"))
+	UPROPERTY(EditAnywhere, Category="Input")
 	UInputAction* SetDestinationTouchAction;
 
 	/** True if the controlled character should navigate to the mouse cursor. */
diff --git a/Templates/TP_TopDown/Source/TP_TopDown/Variant_Strategy/StrategyGameMode.h b/Templates/TP_TopDown/Source/TP_TopDown/Variant_Strategy/StrategyGameMode.h
--- a/Templates/TP_TopDown/Source/TP_TopDown/Variant_Strategy/StrategyGameMode.h
+++ b/Templates/TP_TopDown/Source/TP_TopDown/Variant_Strategy/StrategyGameMode.h
@@ -7,7 +7,7 @@
 #include "StrategyGameMode.generated.h"
 
 /**
- * 
+ *  Simple GameMode for a top down strategy game.
  */
 UCLASS(abstract)
 class AStrategyGameMode : public AGameModeBase
diff --git a/Templates/TP_TopDown/Source/TP_TopDown/Variant_Strategy/StrategyPawn.h b/Templates/TP_TopDown/Source/TP_TopDown/Variant_Strategy/StrategyPawn.h
--- a/Templates/TP_TopDown/Source/TP_TopDown/Variant_Strategy/StrategyPawn.h
+++ b/Templates/TP_TopDown/Source/TP_TopDown/Variant_Strategy/StrategyPawn.h
@@ -19,11 +19,11 @@ class AStrategyPawn : public APawn
 	GENERATED_BODY()
 
 	/** Camera */
-	UPROPERTY(VisibleAnywhere, BlueprintReadOnly, Category = "Components", meta = (AllowPrivateAccess = "true"))
+	UPROPERTY(VisibleAnywhere, BlueprintReadOnly, Category="Components", meta = (AllowPrivateAccess = "true"))
 	UCameraComponent* Camera;
 
 	/** Movement Component */
-	UPROPERTY(VisibleAnywhere, BlueprintReadOnly, Category = "Camera", meta = (AllowPrivateAccess = "true"))
+	UPROPERTY(VisibleAnywhere, BlueprintReadOnly, Category="Components", meta = (AllowPrivateAccess = "true"))
 	UFloatingPawnMovement* FloatingPawnMovement;
 
 public:
diff --git a/Templates/TP_TopDown/Source/TP_TopDown/Variant_Strategy/StrategyPlayerController.h b/Templates/TP_TopDown/Source/TP_TopDown/Variant_Strategy/StrategyPlayerController.h
--- a/Templates/TP_TopDown/Source/TP_TopDown/Variant_Strategy/StrategyPlayerController.h
+++ b/Templates/TP_TopDown/Source/TP_TopDown/Variant_Strategy/StrategyPlayerController.h
@@ -14,6 +14,7 @@ class AStrategyHUD;
 class AStrategyNPC;
 class UInputAction;
 
+/** Enum to determine the last used input type */
 UENUM(BlueprintType)
 enum EStrategyInputMode : uint8
 {
@@ -22,7 +23,9 @@ enum EStrategyInputMode : uint8
 };
 
 /**
- * 
+ *  Player Controller for a top-down strategy game.
+ *  Handles unit selection and commands.
+ *  Implements both mouse and touch controls.
  */
 UCLASS(abstract)
 class AStrategyPlayerController : public APlayerController
@@ -38,15 +41,15 @@ protected:
 	TObjectPtr<AStrategyHUD> StrategyHUD;
 
 	/** Determines the chosen input type */
-	UPROPERTY(EditAnywhere, Category = "Input")
+	UPROPERTY(EditAnywhere, Category="Input")
 	TEnumAsByte<EStrategyInputMode> InputMode = SIM_Mouse;
 
 	/** Input mapping context to use with mouse input */
-	UPROPERTY(EditAnywhere, Category = "Input")
+	UPROPERTY(EditAnywhere, Category="Input")
 	UInputMappingContext* MouseMappingContext;
 
 	/** Input mapping context to use with touch input */
-	UPROPERTY(EditAnywhere, Category = "Input")
+	UPROPERTY(EditAnywhere, Category="Input")
 	UInputMappingContext* TouchMappingContext;
 
 	/** If true, the player is adding or removing units to the selected units list */
@@ -59,59 +62,59 @@ protected:
 	bool bAllowInteraction = true;
 
 	/** Input Action for moving the camera */
-	UPROPERTY(EditAnywhere, Category = "Input")
+	UPROPERTY(EditAnywhere, Category="Input")
 	UInputAction* MoveCameraAction;
 
 	/** Input Action for zooming the camera */
-	UPROPERTY(EditAnywhere, Category = "Input")
+	UPROPERTY(EditAnywhere, Category="Input")
 	UInputAction* ZoomCameraAction;
 
 	/** Input Action for resetting the camera to its default position */
-	UPROPERTY(EditAnywhere, Category = "Input")
+	UPROPERTY(EditAnywhere, Category="Input")
 	UInputAction* ResetCameraAction;
 
 	/** Input Action for select and click */
-	UPROPERTY(EditAnywhere, Category = "Input")
+	UPROPERTY(EditAnywhere, Category="Input")
 	UInputAction* SelectClickAction;
 
 	/** Input Action for select press and hold */
-	UPROPERTY(EditAnywhere, Category = "Input")
+	UPROPERTY(EditAnywhere, Category="Input")
 	UInputAction* SelectHoldAction;
 
 	/** Input Action for click interaction */
-	UPROPERTY(EditAnywhere, Category = "Input")
+	UPROPERTY(EditAnywhere, Category="Input")
 	UInputAction* InteractClickAction;
 
 	/** Input Action for interaction press and hold */
-	UPROPERTY(EditAnywhere, Category = "Input")
+	UPROPERTY(EditAnywhere, Category="Input")
 	UInputAction* InteractHoldAction;
 
 	/** Input Action for modifying selection mode */
-	UPROPERTY(EditAnywhere, Category = "Input")
+	UPROPERTY(EditAnywhere, Category="Input")
 	UInputAction* SelectionModifierAction;
 
 	/** Input Action for primary touch tap */
-	UPROPERTY(EditAnywhere, Category = "Input")
+	UPROPERTY(EditAnywhere, Category="Input")
 	UInputAction* TouchPrimaryTapAction;
 
 	/** Input Action for primary touch hold */
-	UPROPERTY(EditAnywhere, Category = "Input")
+	UPROPERTY(EditAnywhere, Category="Input")
 	UInputAction* TouchPrimaryHoldAction;
 
 	/** Input Action for secondary touch */
-	UPROPERTY(EditAnywhere, Category = "Input")
+	UPROPERTY(EditAnywhere, Category="Input")
 	UInputAction* TouchSecondaryAction;
 
 	/** Input Action for touch double tap */
-	UPROPERTY(EditAnywhere, Category = "Input")
+	UPROPERTY(EditAnywhere, Category="Input")
 	UInputAction* TouchDoubleTapAction;
 
 	/** Max distance to look for nearby units when doing a click or touch interaction */
-	UPROPERTY(EditAnywhere, Category = "Input", meta = (ClampMin = 0, ClampMax = 10000, Units = "cm"))
+	UPROPERTY(EditAnywhere, Category="Input", meta = (ClampMin = 0, ClampMax = 10000, Units = "cm"))
 	float InteractionRadius = 250.0f;
 
 	/** Max distance between the starting and current position of the second touch finger to be considered a box selection */
-	UPROPERTY(EditAnywhere, Category = "Input", meta = (ClampMin = 0, ClampMax = 10000))
+	UPROPERTY(EditAnywhere, Category="Input", meta = (ClampMin = 0, ClampMax = 10000))
 	float MinSecondFingerDistanceForBoxSelect = 10.0f;
 
 	/** Saves the world location of the last initiated interaction */
@@ -277,7 +280,7 @@ protected:
 	FVector ProjectTouchPointToWorldSpace();
 
 	/** Spawns the positive cursor effect */
-	UFUNCTION(BlueprintImplementableEvent, Category="Cursor", meta=(DisplayName="Cursor Feedback"))
+	UFUNCTION(BlueprintImplementableEvent, Category="Cursor", meta = (DisplayName="Cursor Feedback"))
 	void BP_CursorFeedback(FVector Location, bool bPositive);
 
 	/** Resets the interaction flag */
diff --git a/Templates/TP_TopDown/Source/TP_TopDown/Variant_TwinStick/AI/TwinStickAIController.h b/Templates/TP_TopDown/Source/TP_TopDown/Variant_TwinStick/AI/TwinStickAIController.h
--- a/Templates/TP_TopDown/Source/TP_TopDown/Variant_TwinStick/AI/TwinStickAIController.h
+++ b/Templates/TP_TopDown/Source/TP_TopDown/Variant_TwinStick/AI/TwinStickAIController.h
@@ -18,7 +18,7 @@ class ATwinStickAIController : public AAIController
 	GENERATED_BODY()
 	
 	/** StateTree Component */
-	UPROPERTY(VisibleAnywhere, BlueprintReadOnly, Category = "AI", meta = (AllowPrivateAccess = "true"))
+	UPROPERTY(VisibleAnywhere, BlueprintReadOnly, Category="Components", meta = (AllowPrivateAccess = "true"))
 	UStateTreeAIComponent* StateTreeAI;
 
 public:
diff --git a/Templates/TP_TopDown/Source/TP_TopDown/Variant_TwinStick/AI/TwinStickNPC.h b/Templates/TP_TopDown/Source/TP_TopDown/Variant_TwinStick/AI/TwinStickNPC.h
--- a/Templates/TP_TopDown/Source/TP_TopDown/Variant_TwinStick/AI/TwinStickNPC.h
+++ b/Templates/TP_TopDown/Source/TP_TopDown/Variant_TwinStick/AI/TwinStickNPC.h
@@ -38,7 +38,7 @@ protected:
 	TSubclassOf<ATwinStickNPCDestruction> DestructionProxyClass;
 
 	/** Time to wait after this NPC is hit before destroying it */
-	UPROPERTY(EditAnywhere, Category="Pickup", meta=(ClampMin = 0, ClampMax = 5, Units = "Seconds"))
+	UPROPERTY(EditAnywhere, Category="Pickup", meta=(ClampMin = 0, ClampMax = 5, Units = "s"))
 	float DeferredDestructionTime = 0.1f;
 
 	/** Deferred destruction timer */
@@ -46,7 +46,7 @@ protected:
 
 public:
 
-	/** If true, this NPC has already been hit by a projectile and is being destroyed */
+	/** If true, this NPC has already been hit by a projectile and is being destroyed. Exposed to BP so it can be read by StateTree */
 	UPROPERTY(VisibleAnywhere, BlueprintReadOnly, Category="NPC")
 	bool bHit = false;
 
diff --git a/Templates/TP_TopDown/Source/TP_TopDown/Variant_TwinStick/AI/TwinStickSpawner.h b/Templates/TP_TopDown/Source/TP_TopDown/Variant_TwinStick/AI/TwinStickSpawner.h
new file mode 100644
--- /dev/null
+++ b/Templates/TP_TopDown/Source/TP_TopDown/Variant_TwinStick/AI/TwinStickSpawner.h
@@ -0,0 +1,79 @@
+// Copyright Epic Games, Inc. All Rights Reserved.
+
+#pragma once
+
+#include "CoreMinimal.h"
+#include "GameFramework/Actor.h"
+#include "TwinStickNPC.h"
+#include "TwinStickSpawner.generated.h"
+
+class ARecastNavMesh;
+
+/**
+ *  A simple NPC spawner for a Twin Stick Shooter game
+ */
+UCLASS(abstract)
+class ATwinStickSpawner : public AActor
+{
+	GENERATED_BODY()
+	
+protected:
+
+	/** Type of NPC to spawn */
+	UPROPERTY(EditAnywhere, Category="NPC Spawner")
+	TSubclassOf<ATwinStickNPC> NPCClass;
+	
+	/** Time delay between enemy group spawns */
+	UPROPERTY(EditAnywhere, Category="NPC Spawner", meta = (ClampMin = 0, ClampMax = 20, Units = "s"))
+	float SpawnGroupDelay = 5.0f;
+
+	/** Min time delay between individual NPC spawns */
+	UPROPERTY(EditAnywhere, Category="NPC Spawner", meta = (ClampMin = 0, ClampMax = 2, Units = "s"))
+	float MinSpawnDelay = 0.33f;
+
+	/** Max time delay between individual NPC spawns */
+	UPROPERTY(EditAnywhere, Category="NPC Spawner", meta = (ClampMin = 0, ClampMax = 2, Units = "s"))
+	float MaxSpawnDelay = 0.66f;
+
+	/** Radius around the spawner where it can spawn NPCs */
+	UPROPERTY(EditAnywhere, Category="NPC Spawner", meta = (ClampMin = 0, ClampMax = 20, Units = "cm"))
+	float SpawnRadius = 600.0f;
+
+	/** Number of NPCs to spawn per group */
+	UPROPERTY(EditAnywhere, Category="NPC Spawner", meta = (ClampMin = 0, ClampMax = 10))
+	int32 SpawnGroupSize = 3;
+	
+	/** Number of NPCs spawned in the current group */
+	int32 SpawnCount = 0;
+
+	/** NPC group spawn timer */
+	FTimerHandle SpawnGroupTimer;
+
+	/** NPC spawn timer */
+	FTimerHandle SpawnNPCTimer;
+
+	/** Pointer to the recast nav mesh actor, used to provide NPC spawn locations */
+	TObjectPtr<ARecastNavMesh> NavData;
+
+public:	
+
+	/** Constructor */
+	ATwinStickSpawner();
+
+protected:
+
+	/** Gameplay initialization */
+	virtual void BeginPlay() override;
+
+	/** Gameplay cleanup */
+	virtual void EndPlay(EEndPlayReason::Type EndPlayReason) override;
+
+protected:
+
+	/** Spawns a new NPC group */
+	void SpawnNPCGroup();
+
+	/** Spawns an individual NPC */
+	void SpawnNPC();
+
+};
diff --git a/Templates/TP_TopDown/Source/TP_TopDown/Variant_TwinStick/AI/TwinStickStateTreeUtility.h b/Templates/TP_TopDown/Source/TP_TopDown/Variant_TwinStick/AI/TwinStickStateTreeUtility.h
--- a/Templates/TP_TopDown/Source/TP_TopDown/Variant_TwinStick/AI/TwinStickStateTreeUtility.h
+++ b/Templates/TP_TopDown/Source/TP_TopDown/Variant_TwinStick/AI/TwinStickStateTreeUtility.h
@@ -18,11 +18,11 @@ struct FStateTreeGetPlayerInstanceData
 	GENERATED_BODY()
 
 	/** Character that owns this task */
-	UPROPERTY(EditAnywhere, Category = Context)
+	UPROPERTY(EditAnywhere, Category="Context")
 	TObjectPtr<ACharacter> Character;
 
 	/** Character that owns this task */
-	UPROPERTY(VisibleAnywhere)
+	UPROPERTY(VisibleAnywhere, Category="Output")
 	TObjectPtr<ACharacter> TargetPlayerCharacter;
 };
 
diff --git a/Templates/TP_TopDown/Source/TP_TopDown/Variant_TwinStick/TwinStickCharacter.h b/Templates/TP_TopDown/Source/TP_TopDown/Variant_TwinStick/TwinStickCharacter.h
--- a/Templates/TP_TopDown/Source/TP_TopDown/Variant_TwinStick/TwinStickCharacter.h
+++ b/Templates/TP_TopDown/Source/TP_TopDown/Variant_TwinStick/TwinStickCharacter.h
@@ -64,7 +64,7 @@ protected:
 	TEnumAsByte<ETraceTypeQuery> MouseAimTraceChannel;
 
 	/** Impulse to apply to the character when dashing */
-	UPROPERTY(EditAnywhere, Category="Dash")
+	UPROPERTY(EditAnywhere, Category="Dash", meta = (ClampMin = 0, ClampMax = 10000, Units = "cm/s"))
 	float DashImpulse = 2500.0f;
 
 	/** Type of projectile to spawn when shooting */
@@ -72,7 +72,7 @@ protected:
 	TSubclassOf<ATwinStickProjectile> ProjectileClass;
 
 	/** Distance ahead of the character that the projectile will be spawned at */
-	UPROPERTY(EditAnywhere, Category="Projectile", meta=(ClampMin = 0, ClampMax = 1000, Units = "cm"))
+	UPROPERTY(EditAnywhere, Category="Projectile", meta = (ClampMin = 0, ClampMax = 1000, Units = "cm"))
 	float ProjectileOffset = 100.0f;
 
 	/** Type of AoE attack actor to spawn */
@@ -84,15 +84,15 @@ protected:
 	int32 Items = 1;
 
 	/** Knockback impulse to apply to the character when they're damaged */
-	UPROPERTY(EditAnywhere, Category="Damage", meta=(ClampMin = 0, ClampMax = 1000, Units = "cm"))
+	UPROPERTY(EditAnywhere, Category="Damage", meta = (ClampMin = 0, ClampMax = 1000, Units = "cm"))
 	float KnockbackStrength = 2500.0f;
 
 	/** Time to disallow AoE attacks after one is performed */
-	UPROPERTY(EditAnywhere, Category="AoE", meta=(ClampMin = 0, ClampMax = 10, Units = "s"))
+	UPROPERTY(EditAnywhere, Category="AoE", meta = (ClampMin = 0, ClampMax = 10, Units = "s"))
 	float AoECooldownTime = 1.0f;
 
 	/** Speed to blend between our current rotation and the target aim rotation when stick aiming */
-	UPROPERTY(EditAnywhere, Category="Aim", meta=(ClampMin = 0, ClampMax = 100))
+	UPROPERTY(EditAnywhere, Category="Aim", meta = (ClampMin = 0, ClampMax = 100, Units = "s"))
 	float AimRotationInterpSpeed = 10.0f;
 
 	/** Game time of the last AoE attack */
@@ -181,7 +181,7 @@ public:
 protected:
 
 	/** Allows Blueprint code to react to damage */
-	UFUNCTION(BlueprintImplementableEvent, Category="Damage", meta=(DisplayName = "Damaged"))
+	UFUNCTION(BlueprintImplementableEvent, Category="Damage", meta = (DisplayName = "Damaged"))
 	void BP_Damaged();
 
 public:
diff --git a/Templates/TP_VehicleAdv/Source/TP_VehicleAdv/OffroadCar/TP_VehicleAdvOffroadCar.h b/Templates/TP_VehicleAdv/Source/TP_VehicleAdv/OffroadCar/TP_VehicleAdvOffroadCar.h
--- a/Templates/TP_VehicleAdv/Source/TP_VehicleAdv/OffroadCar/TP_VehicleAdvOffroadCar.h
+++ b/Templates/TP_VehicleAdv/Source/TP_VehicleAdv/OffroadCar/TP_VehicleAdvOffroadCar.h
@@ -15,23 +15,23 @@ class ATP_VehicleAdvOffroadCar : public ATP_VehicleAdvPawn
 	GENERATED_BODY()
 	
 	/** Chassis static mesh */
-	UPROPERTY(VisibleAnywhere, BlueprintReadOnly, Category = Meshes, meta = (AllowPrivateAccess = "true"))
+	UPROPERTY(VisibleAnywhere, BlueprintReadOnly, Category ="Components", meta = (AllowPrivateAccess = "true"))
 	UStaticMeshComponent* Chassis;
 
 	/** FL Tire static mesh */
-	UPROPERTY(VisibleAnywhere, BlueprintReadOnly, Category = Meshes, meta = (AllowPrivateAccess = "true"))
+	UPROPERTY(VisibleAnywhere, BlueprintReadOnly, Category ="Components", meta = (AllowPrivateAccess = "true"))
 	UStaticMeshComponent* TireFrontLeft;
 
 	/** FR Tire static mesh */
-	UPROPERTY(VisibleAnywhere, BlueprintReadOnly, Category = Meshes, meta = (AllowPrivateAccess = "true"))
+	UPROPERTY(VisibleAnywhere, BlueprintReadOnly, Category ="Components", meta = (AllowPrivateAccess = "true"))
 	UStaticMeshComponent* TireFrontRight;
 
 	/** RL Tire static mesh */
-	UPROPERTY(VisibleAnywhere, BlueprintReadOnly, Category = Meshes, meta = (AllowPrivateAccess = "true"))
+	UPROPERTY(VisibleAnywhere, BlueprintReadOnly, Category ="Components", meta = (AllowPrivateAccess = "true"))
 	UStaticMeshComponent* TireRearLeft;
 
 	/** RR Tire static mesh */
-	UPROPERTY(VisibleAnywhere, BlueprintReadOnly, Category = Meshes, meta = (AllowPrivateAccess = "true"))
+	UPROPERTY(VisibleAnywhere, BlueprintReadOnly, Category ="Components", meta = (AllowPrivateAccess = "true"))
 	UStaticMeshComponent* TireRearRight;
 
 public:
diff --git a/Templates/TP_VehicleAdv/Source/TP_VehicleAdv/TP_VehicleAdv.Build.cs b/Templates/TP_VehicleAdv/Source/TP_VehicleAdv/TP_VehicleAdv.Build.cs
--- a/Templates/TP_VehicleAdv/Source/TP_VehicleAdv/TP_VehicleAdv.Build.cs
+++ b/Templates/TP_VehicleAdv/Source/TP_VehicleAdv/TP_VehicleAdv.Build.cs
@@ -24,7 +24,8 @@ public class TP_VehicleAdv : ModuleRules
 			"TP_VehicleAdv/SportsCar",
 			"TP_VehicleAdv/OffroadCar",
 			"TP_VehicleAdv/Variant_Offroad",
-			"TP_VehicleAdv/Variant_TimeTrial"
+			"TP_VehicleAdv/Variant_TimeTrial",
+			"TP_VehicleAdv/Variant_TimeTrial/UI"
 		});
 
 		PrivateDependencyModuleNames.AddRange(new string[] { });
diff --git a/Templates/TP_VehicleAdv/Source/TP_VehicleAdv/TP_VehicleAdv.cpp b/Templates/TP_VehicleAdv/Source/TP_VehicleAdv/TP_VehicleAdv.cpp
--- a/Templates/TP_VehicleAdv/Source/TP_VehicleAdv/TP_VehicleAdv.cpp
+++ b/Templates/TP_VehicleAdv/Source/TP_VehicleAdv/TP_VehicleAdv.cpp
@@ -4,4 +4,5 @@
 #include "Modules/ModuleManager.h"
 
 IMPLEMENT_PRIMARY_GAME_MODULE( FDefaultGameModuleImpl, TP_VehicleAdv, "TP_VehicleAdv" );
- 
\ No newline at end of file
+
+DEFINE_LOG_CATEGORY(LogTP_VehicleAdv)
\ No newline at end of file
diff --git a/Templates/TP_VehicleAdv/Source/TP_VehicleAdv/TP_VehicleAdv.h b/Templates/TP_VehicleAdv/Source/TP_VehicleAdv/TP_VehicleAdv.h
--- a/Templates/TP_VehicleAdv/Source/TP_VehicleAdv/TP_VehicleAdv.h
+++ b/Templates/TP_VehicleAdv/Source/TP_VehicleAdv/TP_VehicleAdv.h
@@ -3,3 +3,6 @@
 #pragma once
 
 #include "CoreMinimal.h"
+
+/** Main log category used across the project */
+DECLARE_LOG_CATEGORY_EXTERN(LogTP_VehicleAdv, Log, All);
\ No newline at end of file
diff --git a/Templates/TP_VehicleAdv/Source/TP_VehicleAdv/TP_VehicleAdvPawn.cpp b/Templates/TP_VehicleAdv/Source/TP_VehicleAdv/TP_VehicleAdvPawn.cpp
--- a/Templates/TP_VehicleAdv/Source/TP_VehicleAdv/TP_VehicleAdvPawn.cpp
+++ b/Templates/TP_VehicleAdv/Source/TP_VehicleAdv/TP_VehicleAdvPawn.cpp
@@ -10,11 +10,10 @@
 #include "EnhancedInputSubsystems.h"
 #include "InputActionValue.h"
 #include "ChaosWheeledVehicleMovementComponent.h"
+#include "TP_VehicleAdv.h"
 
 #define LOCTEXT_NAMESPACE "VehiclePawn"
 
-DEFINE_LOG_CATEGORY(LogTemplateVehicle);
-
 ATP_VehicleAdvPawn::ATP_VehicleAdvPawn()
 {
 	// construct the front camera boom
@@ -88,7 +87,7 @@ void ATP_VehicleAdvPawn::SetupPlayerInputComponent(class UInputComponent* Player
 	}
 	else
 	{
-		UE_LOG(LogTemplateVehicle, Error, TEXT("'%s' Failed to find an Enhanced Input component! This template is built to use the Enhanced Input system. If you intend to use the legacy system, then you will need to update this C++ file."), *GetNameSafe(this));
+		UE_LOG(LogTP_VehicleAdv, Error, TEXT("'%s' Failed to find an Enhanced Input component! This template is built to use the Enhanced Input system. If you intend to use the legacy system, then you will need to update this C++ file."), *GetNameSafe(this));
 	}
 }
 
@@ -248,8 +247,6 @@ void ATP_VehicleAdvPawn::DoResetVehicle()
 
 	GetMesh()->SetPhysicsAngularVelocityInDegrees(FVector::ZeroVector);
 	GetMesh()->SetPhysicsLinearVelocity(FVector::ZeroVector);
-
-	UE_LOG(LogTemplateVehicle, Error, TEXT("Reset Vehicle"));
 }
 
 #undef LOCTEXT_NAMESPACE
\ No newline at end of file
diff --git a/Templates/TP_VehicleAdv/Source/TP_VehicleAdv/TP_VehicleAdvPawn.h b/Templates/TP_VehicleAdv/Source/TP_VehicleAdv/TP_VehicleAdvPawn.h
--- a/Templates/TP_VehicleAdv/Source/TP_VehicleAdv/TP_VehicleAdvPawn.h
+++ b/Templates/TP_VehicleAdv/Source/TP_VehicleAdv/TP_VehicleAdvPawn.h
@@ -12,8 +12,6 @@ class UInputAction;
 class UChaosWheeledVehicleMovementComponent;
 struct FInputActionValue;
 
-DECLARE_LOG_CATEGORY_EXTERN(LogTemplateVehicle, Log, All);
-
 /**
  *  Vehicle Pawn class
  *  Handles common functionality for all vehicle types,
@@ -27,19 +25,19 @@ class ATP_VehicleAdvPawn : public AWheeledVehiclePawn
 	GENERATED_BODY()
 
 	/** Spring Arm for the front camera */
-	UPROPERTY(VisibleAnywhere, BlueprintReadOnly, Category = Camera, meta = (AllowPrivateAccess = "true"))
+	UPROPERTY(VisibleAnywhere, BlueprintReadOnly, Category ="Components", meta = (AllowPrivateAccess = "true"))
 	USpringArmComponent* FrontSpringArm;
 
 	/** Front Camera component */
-	UPROPERTY(VisibleAnywhere, BlueprintReadOnly, Category = Camera, meta = (AllowPrivateAccess = "true"))
+	UPROPERTY(VisibleAnywhere, BlueprintReadOnly, Category ="Components", meta = (AllowPrivateAccess = "true"))
 	UCameraComponent* FrontCamera;
 
 	/** Spring Arm for the back camera */
-	UPROPERTY(VisibleAnywhere, BlueprintReadOnly, Category = Camera, meta = (AllowPrivateAccess = "true"))
+	UPROPERTY(VisibleAnywhere, BlueprintReadOnly, Category ="Components", meta = (AllowPrivateAccess = "true"))
 	USpringArmComponent* BackSpringArm;
 
 	/** Back Camera component */
-	UPROPERTY(VisibleAnywhere, BlueprintReadOnly, Category = Camera, meta = (AllowPrivateAccess = "true"))
+	UPROPERTY(VisibleAnywhere, BlueprintReadOnly, Category ="Components", meta = (AllowPrivateAccess = "true"))
 	UCameraComponent* BackCamera;
 
 	/** Cast pointer to the Chaos Vehicle movement component */
@@ -48,31 +46,31 @@ class ATP_VehicleAdvPawn : public AWheeledVehiclePawn
 protected:
 
 	/** Steering Action */
-	UPROPERTY(EditAnywhere, BlueprintReadOnly, Category = Input)
+	UPROPERTY(EditAnywhere, Category="Input")
 	UInputAction* SteeringAction;
 
 	/** Throttle Action */
-	UPROPERTY(EditAnywhere, BlueprintReadOnly, Category = Input)
+	UPROPERTY(EditAnywhere, Category="Input")
 	UInputAction* ThrottleAction;
 
 	/** Brake Action */
-	UPROPERTY(EditAnywhere, BlueprintReadOnly, Category = Input)
+	UPROPERTY(EditAnywhere, Category="Input")
 	UInputAction* BrakeAction;
 
 	/** Handbrake Action */
-	UPROPERTY(EditAnywhere, BlueprintReadOnly, Category = Input)
+	UPROPERTY(EditAnywhere, Category="Input")
 	UInputAction* HandbrakeAction;
 
 	/** Look Around Action */
-	UPROPERTY(EditAnywhere, BlueprintReadOnly, Category = Input)
+	UPROPERTY(EditAnywhere, Category="Input")
 	UInputAction* LookAroundAction;
 
 	/** Toggle Camera Action */
-	UPROPERTY(EditAnywhere, BlueprintReadOnly, Category = Input)
+	UPROPERTY(EditAnywhere, Category="Input")
 	UInputAction* ToggleCameraAction;
 
 	/** Reset Vehicle Action */
-	UPROPERTY(EditAnywhere, BlueprintReadOnly, Category = Input)
+	UPROPERTY(EditAnywhere, Category="Input")
 	UInputAction* ResetVehicleAction;
 
 	/** Keeps track of which camera is active */
diff --git a/Templates/TP_VehicleAdv/Source/TP_VehicleAdv/TP_VehicleAdvPlayerController.h b/Templates/TP_VehicleAdv/Source/TP_VehicleAdv/TP_VehicleAdvPlayerController.h
--- a/Templates/TP_VehicleAdv/Source/TP_VehicleAdv/TP_VehicleAdvPlayerController.h
+++ b/Templates/TP_VehicleAdv/Source/TP_VehicleAdv/TP_VehicleAdvPlayerController.h
@@ -22,24 +22,24 @@ class ATP_VehicleAdvPlayerController : public APlayerController
 protected:
 
 	/** Input Mapping Context to be used for player input */
-	UPROPERTY(EditAnywhere, BlueprintReadOnly, Category="Input")
+	UPROPERTY(EditAnywhere, Category="Input")
 	UInputMappingContext* InputMappingContext;
 
 	/** If true, the optional steering wheel input mapping context will be registered */
-	UPROPERTY(EditAnywhere, BlueprintReadOnly, Category = "Input")
+	UPROPERTY(EditAnywhere, Category = "Input")
 	bool bUseSteeringWheelControls = false;
 
 	/** Optional Input Mapping Context to be used for steering wheel input.
 	 *  This is added alongside the default Input Mapping Context and does not block other forms of input.
 	 */
-	UPROPERTY(EditAnywhere, BlueprintReadOnly, Category = "Input", meta=(EditCondition = "bUseSteeringWheelControls"))
+	UPROPERTY(EditAnywhere, Category = "Input", meta = (EditCondition = "bUseSteeringWheelControls"))
 	UInputMappingContext* SteeringWheelInputMappingContext;
 
 	/** Pointer to the controlled vehicle pawn */
 	TObjectPtr<ATP_VehicleAdvPawn> VehiclePawn;
 
 	/** Type of the UI to spawn */
-	UPROPERTY(EditAnywhere, BlueprintReadOnly, Category="UI")
+	UPROPERTY(EditAnywhere, Category="UI")
 	TSubclassOf<UTP_VehicleAdvUI> VehicleUIClass;
 
 	/** Pointer to the UI widget */
diff --git a/Templates/TP_VehicleAdv/Source/TP_VehicleAdv/TP_VehicleAdvUI.h b/Templates/TP_VehicleAdv/Source/TP_VehicleAdv/TP_VehicleAdvUI.h
--- a/Templates/TP_VehicleAdv/Source/TP_VehicleAdv/TP_VehicleAdvUI.h
+++ b/Templates/TP_VehicleAdv/Source/TP_VehicleAdv/TP_VehicleAdvUI.h
@@ -19,7 +19,7 @@ class UTP_VehicleAdvUI : public UUserWidget
 protected:
 
 	/** Controls the display of speed in Km/h or MPH */
-	UPROPERTY(EditAnywhere, BlueprintReadOnly, Category = Vehicle)
+	UPROPERTY(EditAnywhere, BlueprintReadOnly, Category="Vehicle")
 	bool bIsMPH = false;
 
 public:
@@ -33,10 +33,10 @@ public:
 protected:
 
 	/** Implemented in Blueprint to display the new speed */
-	UFUNCTION(BlueprintImplementableEvent, Category = Vehicle)
+	UFUNCTION(BlueprintImplementableEvent, Category="Vehicle")
 	void OnSpeedUpdate(float NewSpeed);
 
 	/** Implemented in Blueprint to display the new gear */
-	UFUNCTION(BlueprintImplementableEvent, Category = Vehicle)
+	UFUNCTION(BlueprintImplementableEvent, Category="Vehicle")
 	void OnGearUpdate(int32 NewGear);
 };
diff --git a/Templates/TP_VehicleAdv/Source/TP_VehicleAdv/Variant_TimeTrial/TimeTrialPlayerController.h b/Templates/TP_VehicleAdv/Source/TP_VehicleAdv/Variant_TimeTrial/TimeTrialPlayerController.h
--- a/Templates/TP_VehicleAdv/Source/TP_VehicleAdv/Variant_TimeTrial/TimeTrialPlayerController.h
+++ b/Templates/TP_VehicleAdv/Source/TP_VehicleAdv/Variant_TimeTrial/TimeTrialPlayerController.h
@@ -22,17 +22,17 @@ class ATimeTrialPlayerController : public APlayerController
 protected:
 
 	/** Input Mapping Context to be used for player input */
-	UPROPERTY(EditAnywhere, BlueprintReadOnly, Category="Input")
+	UPROPERTY(EditAnywhere, Category="Input")
 	UInputMappingContext* InputMappingContext;
 
 	/** If true, the optional steering wheel input mapping context will be registered */
-	UPROPERTY(EditAnywhere, BlueprintReadOnly, Category = "Input")
+	UPROPERTY(EditAnywhere, Category = "Input")
 	bool bUseSteeringWheelControls = false;
 
 	/** Optional Input Mapping Context to be used for steering wheel input.
 	 *  This is added alongside the default Input Mapping Context and does not block other forms of input.
 	 */
-	UPROPERTY(EditAnywhere, BlueprintReadOnly, Category = "Input", meta=(EditCondition = "bUseSteeringWheelControls"))
+	UPROPERTY(EditAnywhere, Category = "Input", meta = (EditCondition = "bUseSteeringWheelControls"))
 	UInputMappingContext* SteeringWheelInputMappingContext;
 
 	/** Type of UI widget to spawn*/
@@ -43,7 +43,7 @@ protected:
 	TObjectPtr<UTimeTrialUI> UIWidget;
 
 	/** Type of the UI to spawn */
-	UPROPERTY(EditAnywhere, BlueprintReadOnly, Category="UI")
+	UPROPERTY(EditAnywhere, Category="UI")
 	TSubclassOf<UTP_VehicleAdvUI> VehicleUIClass;
 
 	/** Pointer to the UI widget */
diff --git a/Templates/TP_VehicleAdv/Source/TP_VehicleAdv/Variant_TimeTrial/TimeTrialTrackGate.h b/Templates/TP_VehicleAdv/Source/TP_VehicleAdv/Variant_TimeTrial/TimeTrialTrackGate.h
--- a/Templates/TP_VehicleAdv/Source/TP_VehicleAdv/Variant_TimeTrial/TimeTrialTrackGate.h
+++ b/Templates/TP_VehicleAdv/Source/TP_VehicleAdv/Variant_TimeTrial/TimeTrialTrackGate.h
@@ -18,7 +18,7 @@ class ATimeTrialTrackGate : public AActor
 	GENERATED_BODY()
 	
 	/** Collision Box */
-	UPROPERTY(VisibleAnywhere, BlueprintReadOnly, Category = "Components", meta = (AllowPrivateAccess = "true"))
+	UPROPERTY(VisibleAnywhere, Category = "Components", meta = (AllowPrivateAccess = "true"))
 	UBoxComponent* CollisionBox;
 
 protected:
