Templates
- Fixed an issue that caused some input methods to not be recognized in multiplayer
- Consolidated Input Mappings and Touch Controls on PlayerController classes
- TP_VehicleAdv: Time Trial variant now correctly displays vehicle speed and gear on UI
- TP_VehicleAdv: Sports car wheels now spin correctly
- TP_FirstPerson: Fixed an issue that caused unresponsive Shooter NPCs

#rnx

diff --git a/Templates/TP_FirstPerson/Source/TP_FirstPerson/TP_FirstPersonPlayerController.cpp b/Templates/TP_FirstPerson/Source/TP_FirstPerson/TP_FirstPersonPlayerController.cpp
--- a/Templates/TP_FirstPerson/Source/TP_FirstPerson/TP_FirstPersonPlayerController.cpp
+++ b/Templates/TP_FirstPerson/Source/TP_FirstPerson/TP_FirstPersonPlayerController.cpp
@@ -6,6 +6,8 @@
 #include "Engine/LocalPlayer.h"
 #include "InputMappingContext.h"
 #include "TP_FirstPersonCameraManager.h"
+#include "Blueprint/UserWidget.h"
+#include "TP_FirstPerson.h"
 
 ATP_FirstPersonPlayerController::ATP_FirstPersonPlayerController()
 {
@@ -13,16 +15,54 @@ ATP_FirstPersonPlayerController::ATP_FirstPersonPlayerController()
 	PlayerCameraManagerClass = ATP_FirstPersonCameraManager::StaticClass();
 }
 
+void ATP_FirstPersonPlayerController::BeginPlay()
+{
+	Super::BeginPlay();
+
+	// only spawn touch controls on local player controllers
+	if (bUseMobileControlsWidget && IsLocalPlayerController())
+	{
+		// spawn the mobile controls widget
+		MobileControlsWidget = CreateWidget<UUserWidget>(this, MobileControlsWidgetClass);
+
+		if (MobileControlsWidget)
+		{
+			// add the controls to the player screen
+			MobileControlsWidget->AddToPlayerScreen(0);
+
+		} else {
+
+			UE_LOG(LogTP_FirstPerson, Error, TEXT("Could not spawn mobile controls widget."));
+
+		}
+
+	}
+}
+
 void ATP_FirstPersonPlayerController::SetupInputComponent()
 {
 	Super::SetupInputComponent();
 
-	// Add Input Mapping Context
-	if (UEnhancedInputLocalPlayerSubsystem* Subsystem = ULocalPlayer::GetSubsystem<UEnhancedInputLocalPlayerSubsystem>(GetLocalPlayer()))
+	// only add IMCs for local player controllers
+	if (IsLocalPlayerController())
 	{
-		for (UInputMappingContext* CurrentContext : DefaultMappingContexts)
+		// Add Input Mapping Context
+		if (UEnhancedInputLocalPlayerSubsystem* Subsystem = ULocalPlayer::GetSubsystem<UEnhancedInputLocalPlayerSubsystem>(GetLocalPlayer()))
 		{
-			Subsystem->AddMappingContext(CurrentContext, 0);
+			for (UInputMappingContext* CurrentContext : DefaultMappingContexts)
+			{
+				Subsystem->AddMappingContext(CurrentContext, 0);
+			}
+
+			// only add these IMCs if we're not using mobile touch input
+			if (!bUseMobileControlsWidget)
+			{
+				for (UInputMappingContext* CurrentContext : MobileExcludedMappingContexts)
+				{
+					Subsystem->AddMappingContext(CurrentContext, 0);
+				}
+			}
 		}
 	}
+	
 }
diff --git a/Templates/TP_FirstPerson/Source/TP_FirstPerson/TP_FirstPersonPlayerController.h b/Templates/TP_FirstPerson/Source/TP_FirstPerson/TP_FirstPersonPlayerController.h
--- a/Templates/TP_FirstPerson/Source/TP_FirstPerson/TP_FirstPersonPlayerController.h
+++ b/Templates/TP_FirstPerson/Source/TP_FirstPerson/TP_FirstPersonPlayerController.h
@@ -7,6 +7,7 @@
 #include "TP_FirstPersonPlayerController.generated.h"
 
 class UInputMappingContext;
+class UUserWidget;
 
 /**
  *  Simple first person Player Controller
@@ -26,9 +27,27 @@ public:
 protected:
 
 	/** Input Mapping Contexts */
-	UPROPERTY(EditAnywhere, Category="Input")
+	UPROPERTY(EditAnywhere, Category="Input|Input Mappings")
 	TArray<UInputMappingContext*> DefaultMappingContexts;
 
+	/** Input Mapping Contexts */
+	UPROPERTY(EditAnywhere, Category="Input|Input Mappings")
+	TArray<UInputMappingContext*> MobileExcludedMappingContexts;
+
+	/** If true, a mobile controls widget will be used for touch input */
+	UPROPERTY(EditAnywhere, Category="Input|Touch Controls")
+	bool bUseMobileControlsWidget = false;
+
+	/** Mobile controls widget to spawn */
+	UPROPERTY(EditAnywhere, Category="Input|Touch Controls", meta = (EditCondition = "bUseMobileControlsWidget"))
+	TSubclassOf<UUserWidget> MobileControlsWidgetClass;
+
+	/** Pointer to the mobile controls widget */
+	TObjectPtr<UUserWidget> MobileControlsWidget;
+
+	/** Gameplay initialization */
+	virtual void BeginPlay() override;
+
 	/** Input mapping context setup */
 	virtual void SetupInputComponent() override;
 
diff --git a/Templates/TP_FirstPerson/Source/TP_FirstPerson/Variant_Horror/HorrorPlayerController.cpp b/Templates/TP_FirstPerson/Source/TP_FirstPerson/Variant_Horror/HorrorPlayerController.cpp
--- a/Templates/TP_FirstPerson/Source/TP_FirstPerson/Variant_Horror/HorrorPlayerController.cpp
+++ b/Templates/TP_FirstPerson/Source/TP_FirstPerson/Variant_Horror/HorrorPlayerController.cpp
@@ -8,6 +8,7 @@
 #include "TP_FirstPersonCameraManager.h"
 #include "HorrorCharacter.h"
 #include "HorrorUI.h"
+#include "../TP_FirstPerson.h"
 
 AHorrorPlayerController::AHorrorPlayerController()
 {
@@ -15,34 +16,76 @@ AHorrorPlayerController::AHorrorPlayerController()
 	PlayerCameraManagerClass = ATP_FirstPersonCameraManager::StaticClass();
 }
 
+void AHorrorPlayerController::BeginPlay()
+{
+	Super::BeginPlay();
+
+	// only spawn touch controls on local player controllers
+	if (bUseMobileControlsWidget && IsLocalPlayerController())
+	{
+		// spawn the mobile controls widget
+		MobileControlsWidget = CreateWidget<UUserWidget>(this, MobileControlsWidgetClass);
+
+		if (MobileControlsWidget)
+		{
+			// add the controls to the player screen
+			MobileControlsWidget->AddToPlayerScreen(0);
+
+		} else {
+
+			UE_LOG(LogTP_FirstPerson, Error, TEXT("Could not spawn mobile controls widget."));
+
+		}
+
+	}
+}
+
 void AHorrorPlayerController::OnPossess(APawn* aPawn)
 {
 	Super::OnPossess(aPawn);
 
-	// set up the UI for the character
-	if (AHorrorCharacter* HorrorCharacter = Cast<AHorrorCharacter>(aPawn))
+	// only spawn UI on local player controllers
+	if (IsLocalPlayerController())
 	{
-		// create the UI
-		if (!HorrorUI)
+		// set up the UI for the character
+		if (AHorrorCharacter* HorrorCharacter = Cast<AHorrorCharacter>(aPawn))
 		{
-			HorrorUI = CreateWidget<UHorrorUI>(this, HorrorUIClass);
-			HorrorUI->AddToViewport(0);
-		}
+			// create the UI
+			if (!HorrorUI)
+			{
+				HorrorUI = CreateWidget<UHorrorUI>(this, HorrorUIClass);
+				HorrorUI->AddToViewport(0);
+			}
 
-		HorrorUI->SetupCharacter(HorrorCharacter);
+			HorrorUI->SetupCharacter(HorrorCharacter);
+		}
 	}
+	
 }
 
 void AHorrorPlayerController::SetupInputComponent()
 {
 	Super::SetupInputComponent();
 	
-	// Add Input Mapping Contexts
-	if (UEnhancedInputLocalPlayerSubsystem* Subsystem = ULocalPlayer::GetSubsystem<UEnhancedInputLocalPlayerSubsystem>(GetLocalPlayer()))
+	// only add IMCs for local player controllers
+	if (IsLocalPlayerController())
 	{
-		for (UInputMappingContext* CurrentContext : DefaultMappingContexts)
+		// Add Input Mapping Contexts
+		if (UEnhancedInputLocalPlayerSubsystem* Subsystem = ULocalPlayer::GetSubsystem<UEnhancedInputLocalPlayerSubsystem>(GetLocalPlayer()))
 		{
-			Subsystem->AddMappingContext(CurrentContext, 0);
+			for (UInputMappingContext* CurrentContext : DefaultMappingContexts)
+			{
+				Subsystem->AddMappingContext(CurrentContext, 0);
+			}
+
+			// only add these IMCs if we're not using mobile touch input
+			if (!bUseMobileControlsWidget)
+			{
+				for (UInputMappingContext* CurrentContext : MobileExcludedMappingContexts)
+				{
+					Subsystem->AddMappingContext(CurrentContext, 0);
+				}
+			}
 		}
-	}
+	}	
 }
diff --git a/Templates/TP_FirstPerson/Source/TP_FirstPerson/Variant_Horror/HorrorPlayerController.h b/Templates/TP_FirstPerson/Source/TP_FirstPerson/Variant_Horror/HorrorPlayerController.h
--- a/Templates/TP_FirstPerson/Source/TP_FirstPerson/Variant_Horror/HorrorPlayerController.h
+++ b/Templates/TP_FirstPerson/Source/TP_FirstPerson/Variant_Horror/HorrorPlayerController.h
@@ -22,7 +22,7 @@ class TP_FIRSTPERSON_API AHorrorPlayerController : public APlayerController
 protected:
 
 	/** Type of UI widget to spawn */
-	UPROPERTY(EditAnywhere, Category="Horror")
+	UPROPERTY(EditAnywhere, Category="Horror|UI")
 	TSubclassOf<UHorrorUI> HorrorUIClass;
 
 	/** Pointer to the UI widget */
@@ -35,13 +35,31 @@ public:
 
 protected:
 
-	/** Possessed pawn initialization */
-	virtual void OnPossess(APawn* aPawn) override;
-
 	/** Input Mapping Contexts */
-	UPROPERTY(EditAnywhere, Category ="Input")
+	UPROPERTY(EditAnywhere, Category ="Input|Input Mappings")
 	TArray<UInputMappingContext*> DefaultMappingContexts;
 
+	/** Input Mapping Contexts */
+	UPROPERTY(EditAnywhere, Category="Input|Input Mappings")
+	TArray<UInputMappingContext*> MobileExcludedMappingContexts;
+
+	/** If true, a mobile controls widget will be used for touch input */
+	UPROPERTY(EditAnywhere, Category="Input|Touch Controls")
+	bool bUseMobileControlsWidget = false;
+
+	/** Mobile controls widget to spawn */
+	UPROPERTY(EditAnywhere, Category="Input|Touch Controls", meta = (EditCondition = "bUseMobileControlsWidget"))
+	TSubclassOf<UUserWidget> MobileControlsWidgetClass;
+
+	/** Pointer to the mobile controls widget */
+	TObjectPtr<UUserWidget> MobileControlsWidget;
+
+	/** Gameplay Initialization */
+	virtual void BeginPlay() override;
+
+	/** Possessed pawn initialization */
+	virtual void OnPossess(APawn* aPawn) override;
+
 	/** Input mapping context setup */
 	virtual void SetupInputComponent() override;
 
diff --git a/Templates/TP_FirstPerson/Source/TP_FirstPerson/Variant_Shooter/ShooterPlayerController.cpp b/Templates/TP_FirstPerson/Source/TP_FirstPerson/Variant_Shooter/ShooterPlayerController.cpp
--- a/Templates/TP_FirstPerson/Source/TP_FirstPerson/Variant_Shooter/ShooterPlayerController.cpp
+++ b/Templates/TP_FirstPerson/Source/TP_FirstPerson/Variant_Shooter/ShooterPlayerController.cpp
@@ -9,24 +9,69 @@
 #include "GameFramework/PlayerStart.h"
 #include "ShooterCharacter.h"
 #include "ShooterBulletCounterUI.h"
+#include "TP_FirstPerson.h"
 
 void AShooterPlayerController::BeginPlay()
 {
 	Super::BeginPlay();
 
-	// create the bullet counter widget and add it to the screen
-	BulletCounterUI = CreateWidget<UShooterBulletCounterUI>(this, BulletCounterUIClass);
-	BulletCounterUI->AddToPlayerScreen(0);
+	// only spawn touch controls on local player controllers
+	if (IsLocalPlayerController())
+	{
+		if (bUseMobileControlsWidget)
+		{
+			// spawn the mobile controls widget
+			MobileControlsWidget = CreateWidget<UUserWidget>(this, MobileControlsWidgetClass);
+
+			if (MobileControlsWidget)
+			{
+				// add the controls to the player screen
+				MobileControlsWidget->AddToPlayerScreen(0);
+
+			} else {
+
+				UE_LOG(LogTP_FirstPerson, Error, TEXT("Could not spawn mobile controls widget."));
+
+			}
+		}
+
+		// create the bullet counter widget and add it to the screen
+		BulletCounterUI = CreateWidget<UShooterBulletCounterUI>(this, BulletCounterUIClass);
+
+		if (BulletCounterUI)
+		{
+			BulletCounterUI->AddToPlayerScreen(0);
+
+		} else {
+
+			UE_LOG(LogTP_FirstPerson, Error, TEXT("Could not spawn bullet counter widget."));
+
+		}
+		
+	}
 }
 
 void AShooterPlayerController::SetupInputComponent()
 {
-	// add the input mapping contexts
-	if (UEnhancedInputLocalPlayerSubsystem* Subsystem = ULocalPlayer::GetSubsystem<UEnhancedInputLocalPlayerSubsystem>(GetLocalPlayer()))
+	// only add IMCs for local player controllers
+	if (IsLocalPlayerController())
 	{
-		for (UInputMappingContext* CurrentContext : DefaultMappingContexts)
+		// add the input mapping contexts
+		if (UEnhancedInputLocalPlayerSubsystem* Subsystem = ULocalPlayer::GetSubsystem<UEnhancedInputLocalPlayerSubsystem>(GetLocalPlayer()))
 		{
-			Subsystem->AddMappingContext(CurrentContext, 0);
+			for (UInputMappingContext* CurrentContext : DefaultMappingContexts)
+			{
+				Subsystem->AddMappingContext(CurrentContext, 0);
+			}
+
+			// only add these IMCs if we're not using mobile touch input
+			if (!bUseMobileControlsWidget)
+			{
+				for (UInputMappingContext* CurrentContext : MobileExcludedMappingContexts)
+				{
+					Subsystem->AddMappingContext(CurrentContext, 0);
+				}
+			}
 		}
 	}
 }
@@ -74,5 +119,8 @@ void AShooterPlayerController::OnPawnDestroyed(AActor* DestroyedActor)
 void AShooterPlayerController::OnBulletCountUpdated(int32 MagazineSize, int32 Bullets)
 {
 	// update the UI
-	BulletCounterUI->BP_UpdateBulletCounter(MagazineSize, Bullets);
+	if (BulletCounterUI)
+	{
+		BulletCounterUI->BP_UpdateBulletCounter(MagazineSize, Bullets);
+	}
 }
diff --git a/Templates/TP_FirstPerson/Source/TP_FirstPerson/Variant_Shooter/ShooterPlayerController.h b/Templates/TP_FirstPerson/Source/TP_FirstPerson/Variant_Shooter/ShooterPlayerController.h
--- a/Templates/TP_FirstPerson/Source/TP_FirstPerson/Variant_Shooter/ShooterPlayerController.h
+++ b/Templates/TP_FirstPerson/Source/TP_FirstPerson/Variant_Shooter/ShooterPlayerController.h
@@ -23,19 +23,34 @@ class TP_FIRSTPERSON_API AShooterPlayerController : public APlayerController
 protected:
 
 	/** Input mapping contexts for this player */
-	UPROPERTY(EditAnywhere, Category="Shooter")
+	UPROPERTY(EditAnywhere, Category="Input|Input Mappings")
 	TArray<UInputMappingContext*> DefaultMappingContexts;
 
+	/** Input Mapping Contexts */
+	UPROPERTY(EditAnywhere, Category="Input|Input Mappings")
+	TArray<UInputMappingContext*> MobileExcludedMappingContexts;
+
+	/** If true, a mobile controls widget will be used for touch input */
+	UPROPERTY(EditAnywhere, Category="Input|Touch Controls")
+	bool bUseMobileControlsWidget = false;
+
+	/** Mobile controls widget to spawn */
+	UPROPERTY(EditAnywhere, Category="Input|Touch Controls", meta = (EditCondition = "bUseMobileControlsWidget"))
+	TSubclassOf<UUserWidget> MobileControlsWidgetClass;
+
+	/** Pointer to the mobile controls widget */
+	TObjectPtr<UUserWidget> MobileControlsWidget;
+
 	/** Character class to respawn when the possessed pawn is destroyed */
-	UPROPERTY(EditAnywhere, Category="Shooter")
+	UPROPERTY(EditAnywhere, Category="Shooter|Respawn")
 	TSubclassOf<AShooterCharacter> CharacterClass;
 
 	/** Type of bullet counter UI widget to spawn */
-	UPROPERTY(EditAnywhere, Category="Shooter")
+	UPROPERTY(EditAnywhere, Category="Shooter|UI")
 	TSubclassOf<UShooterBulletCounterUI> BulletCounterUIClass;
 
 	/** Tag to grant the possessed pawn to flag it as the player */
-	UPROPERTY(EditAnywhere, Category="Shooter")
+	UPROPERTY(EditAnywhere, Category="Shooter|Player")
 	FName PlayerPawnTag = FName("Player");
 
 	/** Pointer to the bullet counter UI widget */
diff --git a/Templates/TP_ThirdPerson/Source/TP_ThirdPerson/TP_ThirdPersonPlayerController.cpp b/Templates/TP_ThirdPerson/Source/TP_ThirdPerson/TP_ThirdPersonPlayerController.cpp
--- a/Templates/TP_ThirdPerson/Source/TP_ThirdPerson/TP_ThirdPersonPlayerController.cpp
+++ b/Templates/TP_ThirdPerson/Source/TP_ThirdPerson/TP_ThirdPersonPlayerController.cpp
@@ -5,17 +5,56 @@
 #include "EnhancedInputSubsystems.h"
 #include "Engine/LocalPlayer.h"
 #include "InputMappingContext.h"
+#include "Blueprint/UserWidget.h"
+#include "TP_ThirdPerson.h"
+
+void ATP_ThirdPersonPlayerController::BeginPlay()
+{
+	Super::BeginPlay();
+
+	// only spawn touch controls on local player controllers
+	if (bUseMobileControlsWidget && IsLocalPlayerController())
+	{
+		// spawn the mobile controls widget
+		MobileControlsWidget = CreateWidget<UUserWidget>(this, MobileControlsWidgetClass);
+
+		if (MobileControlsWidget)
+		{
+			// add the controls to the player screen
+			MobileControlsWidget->AddToPlayerScreen(0);
+
+		} else {
+
+			UE_LOG(LogTP_ThirdPerson, Error, TEXT("Could not spawn mobile controls widget."));
+
+		}
+
+	}
+}
 
 void ATP_ThirdPersonPlayerController::SetupInputComponent()
 {
 	Super::SetupInputComponent();
 
-	// Add Input Mapping Contexts
-	if (UEnhancedInputLocalPlayerSubsystem* Subsystem = ULocalPlayer::GetSubsystem<UEnhancedInputLocalPlayerSubsystem>(GetLocalPlayer()))
+	// only add IMCs for local player controllers
+	if (IsLocalPlayerController())
 	{
-		for (UInputMappingContext* CurrentContext : DefaultMappingContexts)
+		// Add Input Mapping Contexts
+		if (UEnhancedInputLocalPlayerSubsystem* Subsystem = ULocalPlayer::GetSubsystem<UEnhancedInputLocalPlayerSubsystem>(GetLocalPlayer()))
 		{
-			Subsystem->AddMappingContext(CurrentContext, 0);
+			for (UInputMappingContext* CurrentContext : DefaultMappingContexts)
+			{
+				Subsystem->AddMappingContext(CurrentContext, 0);
+			}
+
+			// only add these IMCs if we're not using mobile touch input
+			if (!bUseMobileControlsWidget)
+			{
+				for (UInputMappingContext* CurrentContext : MobileExcludedMappingContexts)
+				{
+					Subsystem->AddMappingContext(CurrentContext, 0);
+				}
+			}
 		}
 	}
 }
diff --git a/Templates/TP_ThirdPerson/Source/TP_ThirdPerson/TP_ThirdPersonPlayerController.h b/Templates/TP_ThirdPerson/Source/TP_ThirdPerson/TP_ThirdPersonPlayerController.h
--- a/Templates/TP_ThirdPerson/Source/TP_ThirdPerson/TP_ThirdPersonPlayerController.h
+++ b/Templates/TP_ThirdPerson/Source/TP_ThirdPerson/TP_ThirdPersonPlayerController.h
@@ -7,6 +7,7 @@
 #include "TP_ThirdPersonPlayerController.generated.h"
 
 class UInputMappingContext;
+class UUserWidget;
 
 /**
  *  Basic PlayerController class for a third person game
@@ -20,9 +21,27 @@ class ATP_ThirdPersonPlayerController : public APlayerController
 protected:
 
 	/** Input Mapping Contexts */
-	UPROPERTY(EditAnywhere, Category ="Input")
+	UPROPERTY(EditAnywhere, Category ="Input|Input Mappings")
 	TArray<UInputMappingContext*> DefaultMappingContexts;
 
+	/** Input Mapping Contexts */
+	UPROPERTY(EditAnywhere, Category="Input|Input Mappings")
+	TArray<UInputMappingContext*> MobileExcludedMappingContexts;
+
+	/** If true, a mobile controls widget will be used for touch input */
+	UPROPERTY(EditAnywhere, Category="Input|Touch Controls")
+	bool bUseMobileControlsWidget = false;
+
+	/** Mobile controls widget to spawn */
+	UPROPERTY(EditAnywhere, Category="Input|Touch Controls", meta = (EditCondition = "bUseMobileControlsWidget"))
+	TSubclassOf<UUserWidget> MobileControlsWidgetClass;
+
+	/** Pointer to the mobile controls widget */
+	TObjectPtr<UUserWidget> MobileControlsWidget;
+
+	/** Gameplay initialization */
+	virtual void BeginPlay() override;
+
 	/** Input mapping context setup */
 	virtual void SetupInputComponent() override;
 
diff --git a/Templates/TP_ThirdPerson/Source/TP_ThirdPerson/Variant_Combat/CombatPlayerController.cpp b/Templates/TP_ThirdPerson/Source/TP_ThirdPerson/Variant_Combat/CombatPlayerController.cpp
--- a/Templates/TP_ThirdPerson/Source/TP_ThirdPerson/Variant_Combat/CombatPlayerController.cpp
+++ b/Templates/TP_ThirdPerson/Source/TP_ThirdPerson/Variant_Combat/CombatPlayerController.cpp
@@ -9,15 +9,54 @@
 #include "CombatCharacter.h"
 #include "Engine/LocalPlayer.h"
 #include "Engine/World.h"
+#include "Blueprint/UserWidget.h"
+#include "TP_ThirdPerson.h"
+
+void ACombatPlayerController::BeginPlay()
+{
+	Super::BeginPlay();
+
+	// only spawn touch controls on local player controllers
+	if (bUseMobileControlsWidget && IsLocalPlayerController())
+	{
+		// spawn the mobile controls widget
+		MobileControlsWidget = CreateWidget<UUserWidget>(this, MobileControlsWidgetClass);
+
+		if (MobileControlsWidget)
+		{
+			// add the controls to the player screen
+			MobileControlsWidget->AddToPlayerScreen(0);
+
+		} else {
+
+			UE_LOG(LogTP_ThirdPerson, Error, TEXT("Could not spawn mobile controls widget."));
+
+		}
+
+	}
+}
 
 void ACombatPlayerController::SetupInputComponent()
 {
-	// add the input mapping context
-	if (UEnhancedInputLocalPlayerSubsystem* Subsystem = ULocalPlayer::GetSubsystem<UEnhancedInputLocalPlayerSubsystem>(GetLocalPlayer()))
+	// only add IMCs for local player controllers
+	if (IsLocalPlayerController())
 	{
-		for (UInputMappingContext* CurrentContext : DefaultMappingContexts)
+		// add the input mapping context
+		if (UEnhancedInputLocalPlayerSubsystem* Subsystem = ULocalPlayer::GetSubsystem<UEnhancedInputLocalPlayerSubsystem>(GetLocalPlayer()))
 		{
-			Subsystem->AddMappingContext(CurrentContext, 0);
+			for (UInputMappingContext* CurrentContext : DefaultMappingContexts)
+			{
+				Subsystem->AddMappingContext(CurrentContext, 0);
+			}
+
+			// only add these IMCs if we're not using mobile touch input
+			if (!bUseMobileControlsWidget)
+			{
+				for (UInputMappingContext* CurrentContext : MobileExcludedMappingContexts)
+				{
+					Subsystem->AddMappingContext(CurrentContext, 0);
+				}
+			}
 		}
 	}
 }
diff --git a/Templates/TP_ThirdPerson/Source/TP_ThirdPerson/Variant_Combat/CombatPlayerController.h b/Templates/TP_ThirdPerson/Source/TP_ThirdPerson/Variant_Combat/CombatPlayerController.h
--- a/Templates/TP_ThirdPerson/Source/TP_ThirdPerson/Variant_Combat/CombatPlayerController.h
+++ b/Templates/TP_ThirdPerson/Source/TP_ThirdPerson/Variant_Combat/CombatPlayerController.h
@@ -22,9 +22,24 @@ class ACombatPlayerController : public APlayerController
 protected:
 
 	/** Input mapping context for this player */
-	UPROPERTY(EditAnywhere, Category="Input")
+	UPROPERTY(EditAnywhere, Category="Input|Input Mappings")
 	TArray<UInputMappingContext*> DefaultMappingContexts;
 
+	/** Input Mapping Contexts */
+	UPROPERTY(EditAnywhere, Category="Input|Input Mappings")
+	TArray<UInputMappingContext*> MobileExcludedMappingContexts;
+
+	/** If true, a mobile controls widget will be used for touch input */
+	UPROPERTY(EditAnywhere, Category="Input|Touch Controls")
+	bool bUseMobileControlsWidget = false;
+
+	/** Mobile controls widget to spawn */
+	UPROPERTY(EditAnywhere, Category="Input|Touch Controls", meta = (EditCondition = "bUseMobileControlsWidget"))
+	TSubclassOf<UUserWidget> MobileControlsWidgetClass;
+
+	/** Pointer to the mobile controls widget */
+	TObjectPtr<UUserWidget> MobileControlsWidget;
+
 	/** Character class to respawn when the possessed pawn is destroyed */
 	UPROPERTY(EditAnywhere, Category="Respawn")
 	TSubclassOf<ACombatCharacter> CharacterClass;
@@ -34,6 +49,9 @@ protected:
 
 protected:
 
+	/** Gameplay initialization */
+	virtual void BeginPlay() override;
+
 	/** Initialize input bindings */
 	virtual void SetupInputComponent() override;
 
diff --git a/Templates/TP_ThirdPerson/Source/TP_ThirdPerson/Variant_Platforming/PlatformingPlayerController.cpp b/Templates/TP_ThirdPerson/Source/TP_ThirdPerson/Variant_Platforming/PlatformingPlayerController.cpp
--- a/Templates/TP_ThirdPerson/Source/TP_ThirdPerson/Variant_Platforming/PlatformingPlayerController.cpp
+++ b/Templates/TP_ThirdPerson/Source/TP_ThirdPerson/Variant_Platforming/PlatformingPlayerController.cpp
@@ -9,15 +9,54 @@
 #include "PlatformingCharacter.h"
 #include "Engine/LocalPlayer.h"
 #include "Engine/World.h"
+#include "Blueprint/UserWidget.h"
+#include "TP_ThirdPerson.h"
+
+void APlatformingPlayerController::BeginPlay()
+{
+	Super::BeginPlay();
+
+	// only spawn touch controls on local player controllers
+	if (bUseMobileControlsWidget && IsLocalPlayerController())
+	{
+		// spawn the mobile controls widget
+		MobileControlsWidget = CreateWidget<UUserWidget>(this, MobileControlsWidgetClass);
+
+		if (MobileControlsWidget)
+		{
+			// add the controls to the player screen
+			MobileControlsWidget->AddToPlayerScreen(0);
+
+		} else {
+
+			UE_LOG(LogTP_ThirdPerson, Error, TEXT("Could not spawn mobile controls widget."));
+
+		}
+
+	}
+}
 
 void APlatformingPlayerController::SetupInputComponent()
 {
-	// add the input mapping context
-	if (UEnhancedInputLocalPlayerSubsystem* Subsystem = ULocalPlayer::GetSubsystem<UEnhancedInputLocalPlayerSubsystem>(GetLocalPlayer()))
+	// only add IMCs for local player controllers
+	if (IsLocalPlayerController())
 	{
-		for (UInputMappingContext* CurrentContext : DefaultMappingContexts)
+		// add the input mapping context
+		if (UEnhancedInputLocalPlayerSubsystem* Subsystem = ULocalPlayer::GetSubsystem<UEnhancedInputLocalPlayerSubsystem>(GetLocalPlayer()))
 		{
-			Subsystem->AddMappingContext(CurrentContext, 0);
+			for (UInputMappingContext* CurrentContext : DefaultMappingContexts)
+			{
+				Subsystem->AddMappingContext(CurrentContext, 0);
+			}
+
+			// only add these IMCs if we're not using mobile touch input
+			if (!bUseMobileControlsWidget)
+			{
+				for (UInputMappingContext* CurrentContext : MobileExcludedMappingContexts)
+				{
+					Subsystem->AddMappingContext(CurrentContext, 0);
+				}
+			}
 		}
 	}
 }
diff --git a/Templates/TP_ThirdPerson/Source/TP_ThirdPerson/Variant_Platforming/PlatformingPlayerController.h b/Templates/TP_ThirdPerson/Source/TP_ThirdPerson/Variant_Platforming/PlatformingPlayerController.h
--- a/Templates/TP_ThirdPerson/Source/TP_ThirdPerson/Variant_Platforming/PlatformingPlayerController.h
+++ b/Templates/TP_ThirdPerson/Source/TP_ThirdPerson/Variant_Platforming/PlatformingPlayerController.h
@@ -22,15 +22,33 @@ class APlatformingPlayerController : public APlayerController
 protected:
 
 	/** Input mapping context for this player */
-	UPROPERTY(EditAnywhere, Category="Input")
+	UPROPERTY(EditAnywhere, Category="Input|Input Mappings")
 	TArray<UInputMappingContext*> DefaultMappingContexts;
 
+	/** Input Mapping Contexts */
+	UPROPERTY(EditAnywhere, Category="Input|Input Mappings")
+	TArray<UInputMappingContext*> MobileExcludedMappingContexts;
+
+	/** If true, a mobile controls widget will be used for touch input */
+	UPROPERTY(EditAnywhere, Category="Input|Touch Controls")
+	bool bUseMobileControlsWidget = false;
+
+	/** Mobile controls widget to spawn */
+	UPROPERTY(EditAnywhere, Category="Input|Touch Controls", meta = (EditCondition = "bUseMobileControlsWidget"))
+	TSubclassOf<UUserWidget> MobileControlsWidgetClass;
+
+	/** Pointer to the mobile controls widget */
+	TObjectPtr<UUserWidget> MobileControlsWidget;
+
 	/** Character class to respawn when the possessed pawn is destroyed */
 	UPROPERTY(EditAnywhere, Category="Respawn")
 	TSubclassOf<APlatformingCharacter> CharacterClass;
 
 protected:
 
+	/** Gameplay initialization */
+	virtual void BeginPlay() override;
+
 	/** Initialize input bindings */
 	virtual void SetupInputComponent() override;
 
diff --git a/Templates/TP_ThirdPerson/Source/TP_ThirdPerson/Variant_SideScrolling/SideScrollingPlayerController.cpp b/Templates/TP_ThirdPerson/Source/TP_ThirdPerson/Variant_SideScrolling/SideScrollingPlayerController.cpp
--- a/Templates/TP_ThirdPerson/Source/TP_ThirdPerson/Variant_SideScrolling/SideScrollingPlayerController.cpp
+++ b/Templates/TP_ThirdPerson/Source/TP_ThirdPerson/Variant_SideScrolling/SideScrollingPlayerController.cpp
@@ -9,15 +9,54 @@
 #include "SideScrollingCharacter.h"
 #include "Engine/LocalPlayer.h"
 #include "Engine/World.h"
+#include "Blueprint/UserWidget.h"
+#include "TP_ThirdPerson.h"
+
+void ASideScrollingPlayerController::BeginPlay()
+{
+	Super::BeginPlay();
+
+	// only spawn touch controls on local player controllers
+	if (bUseMobileControlsWidget && IsLocalPlayerController())
+	{
+		// spawn the mobile controls widget
+		MobileControlsWidget = CreateWidget<UUserWidget>(this, MobileControlsWidgetClass);
+
+		if (MobileControlsWidget)
+		{
+			// add the controls to the player screen
+			MobileControlsWidget->AddToPlayerScreen(0);
+
+		} else {
+
+			UE_LOG(LogTP_ThirdPerson, Error, TEXT("Could not spawn mobile controls widget."));
+
+		}
+
+	}
+}
 
 void ASideScrollingPlayerController::SetupInputComponent()
 {
-	// add the input mapping context
-	if (UEnhancedInputLocalPlayerSubsystem* Subsystem = ULocalPlayer::GetSubsystem<UEnhancedInputLocalPlayerSubsystem>(GetLocalPlayer()))
+	// only add IMCs for local player controllers
+	if (IsLocalPlayerController())
 	{
-		for (UInputMappingContext* CurrentContext : DefaultMappingContexts)
+		// add the input mapping context
+		if (UEnhancedInputLocalPlayerSubsystem* Subsystem = ULocalPlayer::GetSubsystem<UEnhancedInputLocalPlayerSubsystem>(GetLocalPlayer()))
 		{
-			Subsystem->AddMappingContext(CurrentContext, 0);
+			for (UInputMappingContext* CurrentContext : DefaultMappingContexts)
+			{
+				Subsystem->AddMappingContext(CurrentContext, 0);
+			}
+
+			// only add these IMCs if we're not using mobile touch input
+			if (!bUseMobileControlsWidget)
+			{
+				for (UInputMappingContext* CurrentContext : MobileExcludedMappingContexts)
+				{
+					Subsystem->AddMappingContext(CurrentContext, 0);
+				}
+			}
 		}
 	}
 }
diff --git a/Templates/TP_ThirdPerson/Source/TP_ThirdPerson/Variant_SideScrolling/SideScrollingPlayerController.h b/Templates/TP_ThirdPerson/Source/TP_ThirdPerson/Variant_SideScrolling/SideScrollingPlayerController.h
--- a/Templates/TP_ThirdPerson/Source/TP_ThirdPerson/Variant_SideScrolling/SideScrollingPlayerController.h
+++ b/Templates/TP_ThirdPerson/Source/TP_ThirdPerson/Variant_SideScrolling/SideScrollingPlayerController.h
@@ -23,15 +23,33 @@ class ASideScrollingPlayerController : public APlayerController
 protected:
 
 	/** Input mapping context for this player */
-	UPROPERTY(EditAnywhere, Category="Input")
+	UPROPERTY(EditAnywhere, Category="Input|Input Mappings")
 	TArray<UInputMappingContext*> DefaultMappingContexts;
 
+	/** Input Mapping Contexts */
+	UPROPERTY(EditAnywhere, Category="Input|Input Mappings")
+	TArray<UInputMappingContext*> MobileExcludedMappingContexts;
+
+	/** If true, a mobile controls widget will be used for touch input */
+	UPROPERTY(EditAnywhere, Category="Input|Touch Controls")
+	bool bUseMobileControlsWidget = false;
+
+	/** Mobile controls widget to spawn */
+	UPROPERTY(EditAnywhere, Category="Input|Touch Controls", meta = (EditCondition = "bUseMobileControlsWidget"))
+	TSubclassOf<UUserWidget> MobileControlsWidgetClass;
+
+	/** Pointer to the mobile controls widget */
+	TObjectPtr<UUserWidget> MobileControlsWidget;
+
 	/** Character class to respawn when the possessed pawn is destroyed */
 	UPROPERTY(EditAnywhere, Category="Respawn")
 	TSubclassOf<ASideScrollingCharacter> CharacterClass;
 
 protected:
 
+	/** Gameplay initialization */
+	virtual void BeginPlay() override;
+
 	/** Initialize input bindings */
 	virtual void SetupInputComponent() override;
 
diff --git a/Templates/TP_TopDown/Source/TP_TopDown/TP_TopDownPlayerController.cpp b/Templates/TP_TopDown/Source/TP_TopDown/TP_TopDownPlayerController.cpp
--- a/Templates/TP_TopDown/Source/TP_TopDown/TP_TopDownPlayerController.cpp
+++ b/Templates/TP_TopDown/Source/TP_TopDown/TP_TopDownPlayerController.cpp
@@ -18,48 +18,46 @@ ATP_TopDownPlayerController::ATP_TopDownPlayerController()
 	bIsTouch = false;
 	bMoveToMouseCursor = false;
 
-	// configure the 
+	// configure the controller
 	bShowMouseCursor = true;
 	DefaultMouseCursor = EMouseCursor::Default;
 	CachedDestination = FVector::ZeroVector;
 	FollowTime = 0.f;
 }
 
-void ATP_TopDownPlayerController::BeginPlay()
-{
-	// Call the base class  
-	Super::BeginPlay();
-}
-
 void ATP_TopDownPlayerController::SetupInputComponent()
 {
 	// set up gameplay key bindings
 	Super::SetupInputComponent();
 
-	// Add Input Mapping Context
-	if (UEnhancedInputLocalPlayerSubsystem* Subsystem = ULocalPlayer::GetSubsystem<UEnhancedInputLocalPlayerSubsystem>(GetLocalPlayer()))
+	// Only set up input on local player controllers
+	if (IsLocalPlayerController())
 	{
-		Subsystem->AddMappingContext(DefaultMappingContext, 0);
-	}
+		// Add Input Mapping Context
+		if (UEnhancedInputLocalPlayerSubsystem* Subsystem = ULocalPlayer::GetSubsystem<UEnhancedInputLocalPlayerSubsystem>(GetLocalPlayer()))
+		{
+			Subsystem->AddMappingContext(DefaultMappingContext, 0);
+		}
 
-	// Set up action bindings
-	if (UEnhancedInputComponent* EnhancedInputComponent = Cast<UEnhancedInputComponent>(InputComponent))
-	{
-		// Setup mouse input events
-		EnhancedInputComponent->BindAction(SetDestinationClickAction, ETriggerEvent::Started, this, &ATP_TopDownPlayerController::OnInputStarted);
-		EnhancedInputComponent->BindAction(SetDestinationClickAction, ETriggerEvent::Triggered, this, &ATP_TopDownPlayerController::OnSetDestinationTriggered);
-		EnhancedInputComponent->BindAction(SetDestinationClickAction, ETriggerEvent::Completed, this, &ATP_TopDownPlayerController::OnSetDestinationReleased);
-		EnhancedInputComponent->BindAction(SetDestinationClickAction, ETriggerEvent::Canceled, this, &ATP_TopDownPlayerController::OnSetDestinationReleased);
+		// Set up action bindings
+		if (UEnhancedInputComponent* EnhancedInputComponent = Cast<UEnhancedInputComponent>(InputComponent))
+		{
+			// Setup mouse input events
+			EnhancedInputComponent->BindAction(SetDestinationClickAction, ETriggerEvent::Started, this, &ATP_TopDownPlayerController::OnInputStarted);
+			EnhancedInputComponent->BindAction(SetDestinationClickAction, ETriggerEvent::Triggered, this, &ATP_TopDownPlayerController::OnSetDestinationTriggered);
+			EnhancedInputComponent->BindAction(SetDestinationClickAction, ETriggerEvent::Completed, this, &ATP_TopDownPlayerController::OnSetDestinationReleased);
+			EnhancedInputComponent->BindAction(SetDestinationClickAction, ETriggerEvent::Canceled, this, &ATP_TopDownPlayerController::OnSetDestinationReleased);
 
-		// Setup touch input events
-		EnhancedInputComponent->BindAction(SetDestinationTouchAction, ETriggerEvent::Started, this, &ATP_TopDownPlayerController::OnInputStarted);
-		EnhancedInputComponent->BindAction(SetDestinationTouchAction, ETriggerEvent::Triggered, this, &ATP_TopDownPlayerController::OnTouchTriggered);
-		EnhancedInputComponent->BindAction(SetDestinationTouchAction, ETriggerEvent::Completed, this, &ATP_TopDownPlayerController::OnTouchReleased);
-		EnhancedInputComponent->BindAction(SetDestinationTouchAction, ETriggerEvent::Canceled, this, &ATP_TopDownPlayerController::OnTouchReleased);
-	}
-	else
-	{
-		UE_LOG(LogTP_TopDown, Error, TEXT("'%s' Failed to find an Enhanced Input Component! This template is built to use the Enhanced Input system. If you intend to use the legacy system, then you will need to update this C++ file."), *GetNameSafe(this));
+			// Setup touch input events
+			EnhancedInputComponent->BindAction(SetDestinationTouchAction, ETriggerEvent::Started, this, &ATP_TopDownPlayerController::OnInputStarted);
+			EnhancedInputComponent->BindAction(SetDestinationTouchAction, ETriggerEvent::Triggered, this, &ATP_TopDownPlayerController::OnTouchTriggered);
+			EnhancedInputComponent->BindAction(SetDestinationTouchAction, ETriggerEvent::Completed, this, &ATP_TopDownPlayerController::OnTouchReleased);
+			EnhancedInputComponent->BindAction(SetDestinationTouchAction, ETriggerEvent::Canceled, this, &ATP_TopDownPlayerController::OnTouchReleased);
+		}
+		else
+		{
+			UE_LOG(LogTP_TopDown, Error, TEXT("'%s' Failed to find an Enhanced Input Component! This template is built to use the Enhanced Input system. If you intend to use the legacy system, then you will need to update this C++ file."), *GetNameSafe(this));
+		}
 	}
 }
 
diff --git a/Templates/TP_TopDown/Source/TP_TopDown/TP_TopDownPlayerController.h b/Templates/TP_TopDown/Source/TP_TopDown/TP_TopDownPlayerController.h
--- a/Templates/TP_TopDown/Source/TP_TopDown/TP_TopDownPlayerController.h
+++ b/Templates/TP_TopDown/Source/TP_TopDown/TP_TopDownPlayerController.h
@@ -66,9 +66,6 @@ protected:
 	/** Initialize input bindings */
 	virtual void SetupInputComponent() override;
 	
-	/** Gameplay initialization */
-	virtual void BeginPlay();
-
 	/** Input handlers */
 	void OnInputStarted();
 	void OnSetDestinationTriggered();
diff --git a/Templates/TP_TopDown/Source/TP_TopDown/Variant_Strategy/StrategyPlayerController.cpp b/Templates/TP_TopDown/Source/TP_TopDown/Variant_Strategy/StrategyPlayerController.cpp
--- a/Templates/TP_TopDown/Source/TP_TopDown/Variant_Strategy/StrategyPlayerController.cpp
+++ b/Templates/TP_TopDown/Source/TP_TopDown/Variant_Strategy/StrategyPlayerController.cpp
@@ -27,63 +27,67 @@ void AStrategyPlayerController::SetupInputComponent()
 {
 	Super::SetupInputComponent();
 
-	// add the input mapping context
-	if (UEnhancedInputLocalPlayerSubsystem* Subsystem = ULocalPlayer::GetSubsystem<UEnhancedInputLocalPlayerSubsystem>(GetLocalPlayer()))
+	// only set up input on local player controllers
+	if (IsLocalPlayerController())
 	{
-		// choose the context based on the input mode
-		UInputMappingContext* ChosenContext = nullptr;
-
-		switch (InputMode)
+		// add the input mapping context
+		if (UEnhancedInputLocalPlayerSubsystem* Subsystem = ULocalPlayer::GetSubsystem<UEnhancedInputLocalPlayerSubsystem>(GetLocalPlayer()))
 		{
-		case SIM_Mouse:
-			ChosenContext = MouseMappingContext;
-			break;
-		case SIM_Touch:
-			ChosenContext = TouchMappingContext;
-			break;
-		}
+			// choose the context based on the input mode
+			UInputMappingContext* ChosenContext = nullptr;
 
-		Subsystem->AddMappingContext(ChosenContext, 0);
-	}
+			switch (InputMode)
+			{
+			case SIM_Mouse:
+				ChosenContext = MouseMappingContext;
+				break;
+			case SIM_Touch:
+				ChosenContext = TouchMappingContext;
+				break;
+			}
 
-	// bind the input mappings
-	if (UEnhancedInputComponent* EnhancedInputComponent = Cast<UEnhancedInputComponent>(InputComponent))
-	{
-		// Camera
-		EnhancedInputComponent->BindAction(MoveCameraAction, ETriggerEvent::Triggered, this, &AStrategyPlayerController::MoveCamera);
-		EnhancedInputComponent->BindAction(ZoomCameraAction, ETriggerEvent::Triggered, this, &AStrategyPlayerController::ZoomCamera);
-		EnhancedInputComponent->BindAction(ResetCameraAction, ETriggerEvent::Triggered, this, &AStrategyPlayerController::ResetCamera);
+			Subsystem->AddMappingContext(ChosenContext, 0);
+		}
 
-		// Mouse Interaction
-		EnhancedInputComponent->BindAction(SelectHoldAction, ETriggerEvent::Started, this, &AStrategyPlayerController::SelectHoldStarted);
-		EnhancedInputComponent->BindAction(SelectHoldAction, ETriggerEvent::Triggered, this, &AStrategyPlayerController::SelectHoldTriggered);
-		EnhancedInputComponent->BindAction(SelectHoldAction, ETriggerEvent::Completed, this, &AStrategyPlayerController::SelectHoldCompleted);
-		EnhancedInputComponent->BindAction(SelectHoldAction, ETriggerEvent::Canceled, this, &AStrategyPlayerController::SelectHoldCompleted);
+		// bind the input mappings
+		if (UEnhancedInputComponent* EnhancedInputComponent = Cast<UEnhancedInputComponent>(InputComponent))
+		{
+			// Camera
+			EnhancedInputComponent->BindAction(MoveCameraAction, ETriggerEvent::Triggered, this, &AStrategyPlayerController::MoveCamera);
+			EnhancedInputComponent->BindAction(ZoomCameraAction, ETriggerEvent::Triggered, this, &AStrategyPlayerController::ZoomCamera);
+			EnhancedInputComponent->BindAction(ResetCameraAction, ETriggerEvent::Triggered, this, &AStrategyPlayerController::ResetCamera);
 
-		EnhancedInputComponent->BindAction(SelectClickAction, ETriggerEvent::Completed, this, &AStrategyPlayerController::SelectClick);
+			// Mouse Interaction
+			EnhancedInputComponent->BindAction(SelectHoldAction, ETriggerEvent::Started, this, &AStrategyPlayerController::SelectHoldStarted);
+			EnhancedInputComponent->BindAction(SelectHoldAction, ETriggerEvent::Triggered, this, &AStrategyPlayerController::SelectHoldTriggered);
+			EnhancedInputComponent->BindAction(SelectHoldAction, ETriggerEvent::Completed, this, &AStrategyPlayerController::SelectHoldCompleted);
+			EnhancedInputComponent->BindAction(SelectHoldAction, ETriggerEvent::Canceled, this, &AStrategyPlayerController::SelectHoldCompleted);
 
-		EnhancedInputComponent->BindAction(SelectionModifierAction, ETriggerEvent::Triggered, this, &AStrategyPlayerController::SelectionModifier);
-		EnhancedInputComponent->BindAction(SelectionModifierAction, ETriggerEvent::Completed, this, &AStrategyPlayerController::SelectionModifier);
-		EnhancedInputComponent->BindAction(SelectionModifierAction, ETriggerEvent::Canceled, this, &AStrategyPlayerController::SelectionModifier);
+			EnhancedInputComponent->BindAction(SelectClickAction, ETriggerEvent::Completed, this, &AStrategyPlayerController::SelectClick);
 
-		EnhancedInputComponent->BindAction(InteractHoldAction, ETriggerEvent::Started, this, &AStrategyPlayerController::InteractHoldStarted);
-		EnhancedInputComponent->BindAction(InteractHoldAction, ETriggerEvent::Triggered, this, &AStrategyPlayerController::InteractHoldTriggered);
+			EnhancedInputComponent->BindAction(SelectionModifierAction, ETriggerEvent::Triggered, this, &AStrategyPlayerController::SelectionModifier);
+			EnhancedInputComponent->BindAction(SelectionModifierAction, ETriggerEvent::Completed, this, &AStrategyPlayerController::SelectionModifier);
+			EnhancedInputComponent->BindAction(SelectionModifierAction, ETriggerEvent::Canceled, this, &AStrategyPlayerController::SelectionModifier);
 
-		EnhancedInputComponent->BindAction(InteractClickAction, ETriggerEvent::Started, this, &AStrategyPlayerController::InteractClickStarted);
-		EnhancedInputComponent->BindAction(InteractClickAction, ETriggerEvent::Completed, this, &AStrategyPlayerController::InteractClickCompleted);
+			EnhancedInputComponent->BindAction(InteractHoldAction, ETriggerEvent::Started, this, &AStrategyPlayerController::InteractHoldStarted);
+			EnhancedInputComponent->BindAction(InteractHoldAction, ETriggerEvent::Triggered, this, &AStrategyPlayerController::InteractHoldTriggered);
 
-		// Touch Interaction
-		EnhancedInputComponent->BindAction(TouchPrimaryHoldAction, ETriggerEvent::Started, this, &AStrategyPlayerController::TouchPrimaryHoldStarted);
-		EnhancedInputComponent->BindAction(TouchPrimaryHoldAction, ETriggerEvent::Triggered, this, &AStrategyPlayerController::TouchPrimaryHoldTriggered);
+			EnhancedInputComponent->BindAction(InteractClickAction, ETriggerEvent::Started, this, &AStrategyPlayerController::InteractClickStarted);
+			EnhancedInputComponent->BindAction(InteractClickAction, ETriggerEvent::Completed, this, &AStrategyPlayerController::InteractClickCompleted);
 
-		EnhancedInputComponent->BindAction(TouchPrimaryTapAction, ETriggerEvent::Completed, this, &AStrategyPlayerController::TouchPrimaryTap);
+			// Touch Interaction
+			EnhancedInputComponent->BindAction(TouchPrimaryHoldAction, ETriggerEvent::Started, this, &AStrategyPlayerController::TouchPrimaryHoldStarted);
+			EnhancedInputComponent->BindAction(TouchPrimaryHoldAction, ETriggerEvent::Triggered, this, &AStrategyPlayerController::TouchPrimaryHoldTriggered);
 
-		EnhancedInputComponent->BindAction(TouchSecondaryAction, ETriggerEvent::Started, this, &AStrategyPlayerController::TouchSecondaryStarted);
-		EnhancedInputComponent->BindAction(TouchSecondaryAction, ETriggerEvent::Triggered, this, &AStrategyPlayerController::TouchSecondaryTriggered);
-		EnhancedInputComponent->BindAction(TouchSecondaryAction, ETriggerEvent::Completed, this, &AStrategyPlayerController::TouchSecondaryCompleted);
-		EnhancedInputComponent->BindAction(TouchSecondaryAction, ETriggerEvent::Canceled, this, &AStrategyPlayerController::TouchSecondaryCompleted);
+			EnhancedInputComponent->BindAction(TouchPrimaryTapAction, ETriggerEvent::Completed, this, &AStrategyPlayerController::TouchPrimaryTap);
 
-		EnhancedInputComponent->BindAction(TouchDoubleTapAction, ETriggerEvent::Triggered, this, &AStrategyPlayerController::TouchDoubleTap);
+			EnhancedInputComponent->BindAction(TouchSecondaryAction, ETriggerEvent::Started, this, &AStrategyPlayerController::TouchSecondaryStarted);
+			EnhancedInputComponent->BindAction(TouchSecondaryAction, ETriggerEvent::Triggered, this, &AStrategyPlayerController::TouchSecondaryTriggered);
+			EnhancedInputComponent->BindAction(TouchSecondaryAction, ETriggerEvent::Completed, this, &AStrategyPlayerController::TouchSecondaryCompleted);
+			EnhancedInputComponent->BindAction(TouchSecondaryAction, ETriggerEvent::Canceled, this, &AStrategyPlayerController::TouchSecondaryCompleted);
+
+			EnhancedInputComponent->BindAction(TouchDoubleTapAction, ETriggerEvent::Triggered, this, &AStrategyPlayerController::TouchDoubleTap);
+		}
 	}
 }
 
diff --git a/Templates/TP_TopDown/Source/TP_TopDown/Variant_TwinStick/TwinStickPlayerController.cpp b/Templates/TP_TopDown/Source/TP_TopDown/Variant_TwinStick/TwinStickPlayerController.cpp
--- a/Templates/TP_TopDown/Source/TP_TopDown/Variant_TwinStick/TwinStickPlayerController.cpp
+++ b/Templates/TP_TopDown/Source/TP_TopDown/Variant_TwinStick/TwinStickPlayerController.cpp
@@ -9,13 +9,48 @@
 #include "TwinStickCharacter.h"
 #include "Engine/LocalPlayer.h"
 #include "Engine/World.h"
+#include "Blueprint/UserWidget.h"
+#include "TP_TopDown.h"
+
+void ATwinStickPlayerController::BeginPlay()
+{
+	Super::BeginPlay();
+
+	// only spawn touch controls on local player controllers
+	if (bUseMobileControlsWidget && IsLocalPlayerController())
+	{
+		// spawn the mobile controls widget
+		MobileControlsWidget = CreateWidget<UUserWidget>(this, MobileControlsWidgetClass);
+
+		if (MobileControlsWidget)
+		{
+			// add the controls to the player screen
+			MobileControlsWidget->AddToPlayerScreen(0);
+
+		} else {
+
+			UE_LOG(LogTP_TopDown, Error, TEXT("Could not spawn mobile controls widget."));
+
+		}
+
+	}
+}
 
 void ATwinStickPlayerController::SetupInputComponent()
 {
-	// add the input mapping context
-	if (UEnhancedInputLocalPlayerSubsystem* Subsystem = ULocalPlayer::GetSubsystem<UEnhancedInputLocalPlayerSubsystem>(GetLocalPlayer()))
+	Super::SetupInputComponent();
+
+	// only add IMCs for local player controllers
+	if (IsLocalPlayerController())
 	{
-		Subsystem->AddMappingContext(InputMappingContext, 0);
+		// Add Input Mapping Contexts
+		if (UEnhancedInputLocalPlayerSubsystem* Subsystem = ULocalPlayer::GetSubsystem<UEnhancedInputLocalPlayerSubsystem>(GetLocalPlayer()))
+		{
+			for (UInputMappingContext* CurrentContext : DefaultMappingContexts)
+			{
+				Subsystem->AddMappingContext(CurrentContext, 0);
+			}
+		}
 	}
 }
 
diff --git a/Templates/TP_TopDown/Source/TP_TopDown/Variant_TwinStick/TwinStickPlayerController.h b/Templates/TP_TopDown/Source/TP_TopDown/Variant_TwinStick/TwinStickPlayerController.h
--- a/Templates/TP_TopDown/Source/TP_TopDown/Variant_TwinStick/TwinStickPlayerController.h
+++ b/Templates/TP_TopDown/Source/TP_TopDown/Variant_TwinStick/TwinStickPlayerController.h
@@ -21,9 +21,20 @@ class ATwinStickPlayerController : public APlayerController
 	
 protected:
 
-	/** Input mapping context for this player */
-	UPROPERTY(EditAnywhere, Category="Input")
-	UInputMappingContext* InputMappingContext;
+	/** Input Mapping Contexts */
+	UPROPERTY(EditAnywhere, Category ="Input|Input Mappings")
+	TArray<UInputMappingContext*> DefaultMappingContexts;
+
+	/** If true, a mobile controls widget will be used for touch input */
+	UPROPERTY(EditAnywhere, Category="Input|Touch Controls")
+	bool bUseMobileControlsWidget = false;
+
+	/** Mobile controls widget to spawn */
+	UPROPERTY(EditAnywhere, Category="Input|Touch Controls", meta = (EditCondition = "bUseMobileControlsWidget"))
+	TSubclassOf<UUserWidget> MobileControlsWidgetClass;
+
+	/** Pointer to the mobile controls widget */
+	TObjectPtr<UUserWidget> MobileControlsWidget;
 
 	/** Character class to respawn when the possessed pawn is destroyed */
 	UPROPERTY(EditAnywhere, Category="Respawn")
@@ -31,6 +42,9 @@ protected:
 
 protected:
 
+	/** Gameplay initialization */
+	virtual void BeginPlay() override;
+
 	/** Initialize input bindings */
 	virtual void SetupInputComponent() override;
 
diff --git a/Templates/TP_VehicleAdv/Source/TP_VehicleAdv/TP_VehicleAdvPawn.cpp b/Templates/TP_VehicleAdv/Source/TP_VehicleAdv/TP_VehicleAdvPawn.cpp
--- a/Templates/TP_VehicleAdv/Source/TP_VehicleAdv/TP_VehicleAdvPawn.cpp
+++ b/Templates/TP_VehicleAdv/Source/TP_VehicleAdv/TP_VehicleAdvPawn.cpp
@@ -176,12 +176,18 @@ void ATP_VehicleAdvPawn::DoThrottle(float ThrottleValue)
 {
 	// add the input
 	ChaosVehicleMovement->SetThrottleInput(ThrottleValue);
+
+	// reset the brake input
+	ChaosVehicleMovement->SetBrakeInput(0.0f);
 }
 
 void ATP_VehicleAdvPawn::DoBrake(float BrakeValue)
 {
 	// add the input
 	ChaosVehicleMovement->SetBrakeInput(BrakeValue);
+
+	// reset the throttle input
+	ChaosVehicleMovement->SetThrottleInput(0.0f);
 }
 
 void ATP_VehicleAdvPawn::DoBrakeStart()
diff --git a/Templates/TP_VehicleAdv/Source/TP_VehicleAdv/TP_VehicleAdvPlayerController.cpp b/Templates/TP_VehicleAdv/Source/TP_VehicleAdv/TP_VehicleAdvPlayerController.cpp
--- a/Templates/TP_VehicleAdv/Source/TP_VehicleAdv/TP_VehicleAdvPlayerController.cpp
+++ b/Templates/TP_VehicleAdv/Source/TP_VehicleAdv/TP_VehicleAdvPlayerController.cpp
@@ -6,33 +6,72 @@
 #include "TP_VehicleAdvUI.h"
 #include "EnhancedInputSubsystems.h"
 #include "ChaosWheeledVehicleMovementComponent.h"
+#include "Blueprint/UserWidget.h"
+#include "TP_VehicleAdv.h"
 
 void ATP_VehicleAdvPlayerController::BeginPlay()
 {
 	Super::BeginPlay();
 	
-	// spawn the UI widget and add it to the viewport
-	VehicleUI = CreateWidget<UTP_VehicleAdvUI>(this, VehicleUIClass);
+	// only spawn UI on local player controllers
+	if (IsLocalPlayerController())
+	{
+		if (bUseMobileControlsWidget)
+		{
+			// spawn the mobile controls widget
+			MobileControlsWidget = CreateWidget<UUserWidget>(this, MobileControlsWidgetClass);
+
+			if (MobileControlsWidget)
+			{
+				// add the controls to the player screen
+				MobileControlsWidget->AddToPlayerScreen(0);
+
+			} else {
 
-	check(VehicleUI);
+				UE_LOG(LogTP_VehicleAdv, Error, TEXT("Could not spawn mobile controls widget."));
 
-	VehicleUI->AddToViewport();
+			}
+		}
+		
+
+		// spawn the UI widget and add it to the viewport
+		VehicleUI = CreateWidget<UTP_VehicleAdvUI>(this, VehicleUIClass);
+
+		if (VehicleUI)
+		{
+			VehicleUI->AddToViewport();
+
+		} else {
+
+			UE_LOG(LogTP_VehicleAdv, Error, TEXT("Could not spawn vehicle UI widget."));
+
+		}
+	}
 }
 
 void ATP_VehicleAdvPlayerController::SetupInputComponent()
 {
 	Super::SetupInputComponent();
 	
-	// get the enhanced input subsystem
-	if (UEnhancedInputLocalPlayerSubsystem* Subsystem = ULocalPlayer::GetSubsystem<UEnhancedInputLocalPlayerSubsystem>(GetLocalPlayer()))
+	// only add IMCs for local player controllers
+	if (IsLocalPlayerController())
 	{
-		// add the mapping context so we get controls
-		Subsystem->AddMappingContext(InputMappingContext, 0);
-
-		// optionally add the steering wheel context
-		if (bUseSteeringWheelControls && SteeringWheelInputMappingContext)
+		// Add Input Mapping Contexts
+		if (UEnhancedInputLocalPlayerSubsystem* Subsystem = ULocalPlayer::GetSubsystem<UEnhancedInputLocalPlayerSubsystem>(GetLocalPlayer()))
 		{
-			Subsystem->AddMappingContext(SteeringWheelInputMappingContext, 1);
+			for (UInputMappingContext* CurrentContext : DefaultMappingContexts)
+			{
+				Subsystem->AddMappingContext(CurrentContext, 0);
+			}
+
+			// only add these IMCs if we're not using mobile touch input
+			if (!bUseMobileControlsWidget)
+			{
+				for (UInputMappingContext* CurrentContext : MobileExcludedMappingContexts)
+				{
+					Subsystem->AddMappingContext(CurrentContext, 0);
+				}
+			}
 		}
 	}
 }
diff --git a/Templates/TP_VehicleAdv/Source/TP_VehicleAdv/TP_VehicleAdvPlayerController.h b/Templates/TP_VehicleAdv/Source/TP_VehicleAdv/TP_VehicleAdvPlayerController.h
--- a/Templates/TP_VehicleAdv/Source/TP_VehicleAdv/TP_VehicleAdvPlayerController.h
+++ b/Templates/TP_VehicleAdv/Source/TP_VehicleAdv/TP_VehicleAdvPlayerController.h
@@ -21,25 +21,40 @@ class ATP_VehicleAdvPlayerController : public APlayerController
 
 protected:
 
-	/** Input Mapping Context to be used for player input */
-	UPROPERTY(EditAnywhere, Category="Input")
-	UInputMappingContext* InputMappingContext;
+	/** Input Mapping Contexts */
+	UPROPERTY(EditAnywhere, Category ="Input|Input Mappings")
+	TArray<UInputMappingContext*> DefaultMappingContexts;
+
+	/** Input Mapping Contexts */
+	UPROPERTY(EditAnywhere, Category="Input|Input Mappings")
+	TArray<UInputMappingContext*> MobileExcludedMappingContexts;
+
+	/** If true, a mobile controls widget will be used for touch input */
+	UPROPERTY(EditAnywhere, Category="Input|Touch Controls")
+	bool bUseMobileControlsWidget = false;
+
+	/** Mobile controls widget to spawn */
+	UPROPERTY(EditAnywhere, Category="Input|Touch Controls", meta = (EditCondition = "bUseMobileControlsWidget"))
+	TSubclassOf<UUserWidget> MobileControlsWidgetClass;
+
+	/** Pointer to the mobile controls widget */
+	TObjectPtr<UUserWidget> MobileControlsWidget;
 
 	/** If true, the optional steering wheel input mapping context will be registered */
-	UPROPERTY(EditAnywhere, Category = "Input")
+	UPROPERTY(EditAnywhere, Category = "Input|Steering Wheel Controls")
 	bool bUseSteeringWheelControls = false;
 
 	/** Optional Input Mapping Context to be used for steering wheel input.
 	 *  This is added alongside the default Input Mapping Context and does not block other forms of input.
 	 */
-	UPROPERTY(EditAnywhere, Category = "Input", meta = (EditCondition = "bUseSteeringWheelControls"))
+	UPROPERTY(EditAnywhere, Category = "Input|Steering Wheel Controls", meta = (EditCondition = "bUseSteeringWheelControls"))
 	UInputMappingContext* SteeringWheelInputMappingContext;
 
 	/** Pointer to the controlled vehicle pawn */
 	TObjectPtr<ATP_VehicleAdvPawn> VehiclePawn;
 
 	/** Type of the UI to spawn */
-	UPROPERTY(EditAnywhere, Category="UI")
+	UPROPERTY(EditAnywhere, Category="Vehicle|UI")
 	TSubclassOf<UTP_VehicleAdvUI> VehicleUIClass;
 
 	/** Pointer to the UI widget */
diff --git a/Templates/TP_VehicleAdv/Source/TP_VehicleAdv/Variant_TimeTrial/TimeTrialPlayerController.cpp b/Templates/TP_VehicleAdv/Source/TP_VehicleAdv/Variant_TimeTrial/TimeTrialPlayerController.cpp
--- a/Templates/TP_VehicleAdv/Source/TP_VehicleAdv/Variant_TimeTrial/TimeTrialPlayerController.cpp
+++ b/Templates/TP_VehicleAdv/Source/TP_VehicleAdv/Variant_TimeTrial/TimeTrialPlayerController.cpp
@@ -10,47 +10,115 @@
 #include "Engine/LocalPlayer.h"
 #include "InputMappingContext.h"
 #include "TP_VehicleAdvUI.h"
+#include "TP_VehicleAdvPawn.h"
+#include "ChaosWheeledVehicleMovementComponent.h"
+#include "Blueprint/UserWidget.h"
+#include "TP_VehicleAdv.h"
 
 void ATimeTrialPlayerController::BeginPlay()
 {
 	Super::BeginPlay();
 
-	// create the UI widget
-	UIWidget = CreateWidget<UTimeTrialUI>(this, UIWidgetClass);
-	UIWidget->AddToViewport(0);
+	// only spawn UI on local player controllers
+	if (IsLocalPlayerController())
+	{
+		if (bUseMobileControlsWidget)
+		{
+			// spawn the mobile controls widget
+			MobileControlsWidget = CreateWidget<UUserWidget>(this, MobileControlsWidgetClass);
+
+			if (MobileControlsWidget)
+			{
+				// add the controls to the player screen
+				MobileControlsWidget->AddToPlayerScreen(0);
+
+			} else {
+
+				UE_LOG(LogTP_VehicleAdv, Error, TEXT("Could not spawn mobile controls widget."));
+
+			}
+		}
+
+		// create the UI widget
+		UIWidget = CreateWidget<UTimeTrialUI>(this, UIWidgetClass);
+
+		if (UIWidget)
+		{
+			UIWidget->AddToViewport(0);
+
+			// subscribe to the race start delegate
+			UIWidget->OnRaceStart.AddDynamic(this, &ATimeTrialPlayerController::StartRace);
+
+		} else {
+
+			UE_LOG(LogTP_VehicleAdv, Error, TEXT("Could not spawn Time Trial UI widget."));
+
+		}
+		
 
-	// spawn the UI widget and add it to the viewport
-	VehicleUI = CreateWidget<UTP_VehicleAdvUI>(this, VehicleUIClass);
-	VehicleUI->AddToViewport();
+		// spawn the UI widget and add it to the viewport
+		VehicleUI = CreateWidget<UTP_VehicleAdvUI>(this, VehicleUIClass);
+
+		if (VehicleUI)
+		{
+			VehicleUI->AddToViewport(0);
+
+		} else {
+
+			UE_LOG(LogTP_VehicleAdv, Error, TEXT("Could not spawn vehicle UI widget."));
+
+		}
+	}
 
-	// subscribe to the race start delegate
-	UIWidget->OnRaceStart.AddDynamic(this, &ATimeTrialPlayerController::StartRace);
 }
 
 void ATimeTrialPlayerController::SetupInputComponent()
 {
 	Super::SetupInputComponent();
 
-	// get the enhanced input subsystem
-	if (UEnhancedInputLocalPlayerSubsystem* Subsystem = ULocalPlayer::GetSubsystem<UEnhancedInputLocalPlayerSubsystem>(GetLocalPlayer()))
+	// only add IMCs for local player controllers
+	if (IsLocalPlayerController())
 	{
-		// add the mapping context so we get controls
-		Subsystem->AddMappingContext(InputMappingContext, 0);
-
-		// optionally add the steering wheel context
-		if (bUseSteeringWheelControls && SteeringWheelInputMappingContext)
+		// Add Input Mapping Contexts
+		if (UEnhancedInputLocalPlayerSubsystem* Subsystem = ULocalPlayer::GetSubsystem<UEnhancedInputLocalPlayerSubsystem>(GetLocalPlayer()))
 		{
-			Subsystem->AddMappingContext(SteeringWheelInputMappingContext, 1);
+			for (UInputMappingContext* CurrentContext : DefaultMappingContexts)
+			{
+				Subsystem->AddMappingContext(CurrentContext, 0);
+			}
+
+			// only add these IMCs if we're not using mobile touch input
+			if (!bUseMobileControlsWidget)
+			{
+				for (UInputMappingContext* CurrentContext : MobileExcludedMappingContexts)
+				{
+					Subsystem->AddMappingContext(CurrentContext, 0);
+				}
+			}
 		}
 	}
 }
 
-void ATimeTrialPlayerController::OnPossess(APawn* aPawn)
+void ATimeTrialPlayerController::OnPossess(APawn* InPawn)
 {
-	Super::OnPossess(aPawn);
+	Super::OnPossess(InPawn);
+
+	// get a pointer to the controlled pawn
+	VehiclePawn = CastChecked<ATP_VehicleAdvPawn>(InPawn);
 
 	// disable input on the pawn
-	aPawn->DisableInput(this);
+	VehiclePawn->DisableInput(this);
+}
+
+void ATimeTrialPlayerController::Tick(float Delta)
+{
+	Super::Tick(Delta);
+
+	if (IsValid(VehiclePawn) && IsValid(VehicleUI))
+	{
+		VehicleUI->UpdateSpeed(VehiclePawn->GetChaosVehicleMovement()->GetForwardSpeed());
+		VehicleUI->UpdateGear(VehiclePawn->GetChaosVehicleMovement()->GetCurrentGear());
+	}
 }
 
 void ATimeTrialPlayerController::StartRace()
diff --git a/Templates/TP_VehicleAdv/Source/TP_VehicleAdv/Variant_TimeTrial/TimeTrialPlayerController.h b/Templates/TP_VehicleAdv/Source/TP_VehicleAdv/Variant_TimeTrial/TimeTrialPlayerController.h
--- a/Templates/TP_VehicleAdv/Source/TP_VehicleAdv/Variant_TimeTrial/TimeTrialPlayerController.h
+++ b/Templates/TP_VehicleAdv/Source/TP_VehicleAdv/Variant_TimeTrial/TimeTrialPlayerController.h
@@ -10,6 +10,7 @@ class ATimeTrialTrackGate;
 class UTimeTrialUI;
 class UInputMappingContext;
 class UTP_VehicleAdvUI;
+class ATP_VehicleAdvPawn;
 
 /**
  *  A simple PlayerController for a Time Trial racing game
@@ -21,29 +22,44 @@ class ATimeTrialPlayerController : public APlayerController
 	
 protected:
 
-	/** Input Mapping Context to be used for player input */
-	UPROPERTY(EditAnywhere, Category="Input")
-	UInputMappingContext* InputMappingContext;
+	/** Input Mapping Contexts */
+	UPROPERTY(EditAnywhere, Category ="Input|Input Mappings")
+	TArray<UInputMappingContext*> DefaultMappingContexts;
+
+	/** Input Mapping Contexts */
+	UPROPERTY(EditAnywhere, Category="Input|Input Mappings")
+	TArray<UInputMappingContext*> MobileExcludedMappingContexts;
+
+	/** If true, a mobile controls widget will be used for touch input */
+	UPROPERTY(EditAnywhere, Category="Input|Touch Controls")
+	bool bUseMobileControlsWidget = false;
+
+	/** Mobile controls widget to spawn */
+	UPROPERTY(EditAnywhere, Category="Input|Touch Controls", meta = (EditCondition = "bUseMobileControlsWidget"))
+	TSubclassOf<UUserWidget> MobileControlsWidgetClass;
+
+	/** Pointer to the mobile controls widget */
+	TObjectPtr<UUserWidget> MobileControlsWidget;
 
 	/** If true, the optional steering wheel input mapping context will be registered */
-	UPROPERTY(EditAnywhere, Category = "Input")
+	UPROPERTY(EditAnywhere, Category = "Input|Steering Wheel Controls")
 	bool bUseSteeringWheelControls = false;
 
 	/** Optional Input Mapping Context to be used for steering wheel input.
 	 *  This is added alongside the default Input Mapping Context and does not block other forms of input.
 	 */
-	UPROPERTY(EditAnywhere, Category = "Input", meta = (EditCondition = "bUseSteeringWheelControls"))
+	UPROPERTY(EditAnywhere, Category = "Input|Steering Wheel Controls", meta = (EditCondition = "bUseSteeringWheelControls"))
 	UInputMappingContext* SteeringWheelInputMappingContext;
 
 	/** Type of UI widget to spawn*/
-	UPROPERTY(EditAnywhere, Category="Time Trial")
+	UPROPERTY(EditAnywhere, Category="Time Trial|UI")
 	TSubclassOf<UTimeTrialUI> UIWidgetClass;
 
 	/** Pointer to the UI Widget */
 	TObjectPtr<UTimeTrialUI> UIWidget;
 
 	/** Type of the UI to spawn */
-	UPROPERTY(EditAnywhere, Category="UI")
+	UPROPERTY(EditAnywhere, Category="Vehicle|UI")
 	TSubclassOf<UTP_VehicleAdvUI> VehicleUIClass;
 
 	/** Pointer to the UI widget */
@@ -55,6 +71,9 @@ protected:
 	/** Lap counter */
 	int32 CurrentLap = 0;
 
+	/** Pointer to the controlled vehicle pawn */
+	TObjectPtr<ATP_VehicleAdvPawn> VehiclePawn;
+
 protected:
 
 	/** Gameplay initialization */
@@ -66,6 +85,11 @@ protected:
 	/** Pawn initialization */
 	virtual void OnPossess(APawn* aPawn) override;
 
+public:
+
+	/** UI vehicle state update on tick */
+	virtual void Tick(float Delta) override;
+
 public:
 
 	/** Sets up the race start */
