Templates
- TP_TopDown: Implemented missing autofire feature on TwinStick C++ version.
- TP_TopDown, BP: Adjusted RecastNavMesh on Lvl_TwinStick to try and cover walls that are unreachable to enemies.

#rnx

diff --git a/Templates/TP_TopDown/Source/TP_TopDown/Variant_TwinStick/TwinStickCharacter.cpp b/Templates/TP_TopDown/Source/TP_TopDown/Variant_TwinStick/TwinStickCharacter.cpp
--- a/Templates/TP_TopDown/Source/TP_TopDown/Variant_TwinStick/TwinStickCharacter.cpp
+++ b/Templates/TP_TopDown/Source/TP_TopDown/Variant_TwinStick/TwinStickCharacter.cpp
@@ -11,6 +11,8 @@
 #include "TwinStickAoEAttack.h"
 #include "Kismet/KismetMathLibrary.h"
 #include "TwinStickProjectile.h"
+#include "Engine/World.h"
+#include "TimerManager.h"
 
 ATwinStickCharacter::ATwinStickCharacter()
 {
@@ -52,6 +54,14 @@ void ATwinStickCharacter::BeginPlay()
 	UpdateItems();
 }
 
+void ATwinStickCharacter::EndPlay(EEndPlayReason::Type EndPlayReason)
+{
+	Super::EndPlay(EndPlayReason);
+
+	/** Clear the autofire timer */
+	GetWorld()->GetTimerManager().ClearTimer(AutoFireTimer);
+}
+
 void ATwinStickCharacter::NotifyControllerChanged()
 {
 	Super::NotifyControllerChanged();
@@ -196,6 +206,19 @@ void ATwinStickCharacter::DoAim(float AxisX, float AxisY)
 	{
 		PlayerController->SetShowMouseCursor(false);
 	}
+
+	// are we on autofire cooldown?
+	if (!bAutoFireActive)
+	{
+		// set ourselves on cooldown
+		bAutoFireActive = true;
+
+		// fire a projectile
+		DoShoot();
+
+		// schedule autofire cooldown reset
+		GetWorld()->GetTimerManager().SetTimer(AutoFireTimer, this, &ATwinStickCharacter::ResetAutoFire, AutoFireDelay, false);
+	}
 }
 
 void ATwinStickCharacter::DoDash()
@@ -279,3 +302,9 @@ void ATwinStickCharacter::UpdateItems()
 	}
 }
 
+void ATwinStickCharacter::ResetAutoFire()
+{
+	// reset the autofire flag
+	bAutoFireActive = false;
+}
+
diff --git a/Templates/TP_TopDown/Source/TP_TopDown/Variant_TwinStick/TwinStickCharacter.h b/Templates/TP_TopDown/Source/TP_TopDown/Variant_TwinStick/TwinStickCharacter.h
--- a/Templates/TP_TopDown/Source/TP_TopDown/Variant_TwinStick/TwinStickCharacter.h
+++ b/Templates/TP_TopDown/Source/TP_TopDown/Variant_TwinStick/TwinStickCharacter.h
@@ -110,6 +110,16 @@ protected:
 	/** Last held move input */
 	FVector2D LastMoveInput;
 
+	/** If true, the player is auto firing while stick aiming */
+	bool bAutoFireActive = false;
+
+	/** Time to wait between autofire attempts */
+	UPROPERTY(EditAnywhere, Category="Aim", meta = (ClampMin = 0, ClampMax = 5, Units = "s"))
+	float AutoFireDelay = 0.2f;
+
+	/** Timer to handle stick autofire */
+	FTimerHandle AutoFireTimer;
+
 public:
 	
 	/** Constructor */
@@ -120,6 +130,9 @@ protected:
 	/** Gameplay Initialization */
 	virtual void BeginPlay() override;
 
+	/** Gameplay cleanup */
+	virtual void EndPlay(EEndPlayReason::Type EndPlayReason) override;
+
 	/** Possessed by controller initialization */
 	virtual void NotifyControllerChanged() override;
 
@@ -193,4 +206,7 @@ protected:
 
 	/** Updates the items counter on the Game Mode */
 	void UpdateItems();
+
+	/** Resets stick the aim autofire flag after the autofire timer has expired */
+	void ResetAutoFire();
 };
