TP_TopDown, TP_TopDownBP
- Fixed a potential crash due to BeginPlay timing between game framework classes.
- Changed the implementation of the AoE explosion damage checks on TwinStick variant
- Fixed base template character not responding to point and click inputs.
- TP_TopDownBP still requires a map update for the point and click move fix to fully work.

#rnx

#virtualized

diff --git a/Templates/TP_TopDown/Source/TP_TopDown/TP_TopDownCharacter.h b/Templates/TP_TopDown/Source/TP_TopDown/TP_TopDownCharacter.h
--- a/Templates/TP_TopDown/Source/TP_TopDown/TP_TopDownCharacter.h
+++ b/Templates/TP_TopDown/Source/TP_TopDown/TP_TopDownCharacter.h
@@ -6,6 +6,9 @@
 #include "GameFramework/Character.h"
 #include "TP_TopDownCharacter.generated.h"
 
+class UCameraComponent;
+class USpringArmComponent;
+
 /**
  *  A controllable top-down perspective character
  */
@@ -18,11 +21,11 @@ private:
 
 	/** Top down camera */
 	UPROPERTY(VisibleAnywhere, BlueprintReadOnly, Category="Components", meta = (AllowPrivateAccess = "true"))
-	class UCameraComponent* TopDownCameraComponent;
+	TObjectPtr<UCameraComponent> TopDownCameraComponent;
 
 	/** Camera boom positioning the camera above the character */
 	UPROPERTY(VisibleAnywhere, BlueprintReadOnly, Category="Components", meta = (AllowPrivateAccess = "true"))
-	class USpringArmComponent* CameraBoom;
+	TObjectPtr<USpringArmComponent> CameraBoom;
 
 public:
 
@@ -36,10 +39,10 @@ public:
 	virtual void Tick(float DeltaSeconds) override;
 
 	/** Returns the camera component **/
-	FORCEINLINE class UCameraComponent* GetTopDownCameraComponent() const { return TopDownCameraComponent; }
+	UCameraComponent* GetTopDownCameraComponent() const { return TopDownCameraComponent.Get(); }
 
 	/** Returns the Camera Boom component **/
-	FORCEINLINE class USpringArmComponent* GetCameraBoom() const { return CameraBoom; }
+	USpringArmComponent* GetCameraBoom() const { return CameraBoom.Get(); }
 
 };
 
diff --git a/Templates/TP_TopDown/Source/TP_TopDown/TP_TopDownPlayerController.cpp b/Templates/TP_TopDown/Source/TP_TopDown/TP_TopDownPlayerController.cpp
--- a/Templates/TP_TopDown/Source/TP_TopDown/TP_TopDownPlayerController.cpp
+++ b/Templates/TP_TopDown/Source/TP_TopDown/TP_TopDownPlayerController.cpp
@@ -8,6 +8,7 @@
 #include "TP_TopDownCharacter.h"
 #include "Engine/World.h"
 #include "EnhancedInputComponent.h"
+#include "Navigation/PathFollowingComponent.h"
 #include "InputActionValue.h"
 #include "EnhancedInputSubsystems.h"
 #include "Engine/LocalPlayer.h"
@@ -18,6 +19,9 @@ ATP_TopDownPlayerController::ATP_TopDownPlayerController()
 	bIsTouch = false;
 	bMoveToMouseCursor = false;
 
+	// create the path following comp
+	PathFollowingComponent = CreateDefaultSubobject<UPathFollowingComponent>(TEXT("Path Following Component"));
+
 	// configure the controller
 	bShowMouseCursor = true;
 	DefaultMouseCursor = EMouseCursor::Default;
@@ -64,6 +68,9 @@ void ATP_TopDownPlayerController::SetupInputComponent()
 void ATP_TopDownPlayerController::OnInputStarted()
 {
 	StopMovement();
+
+	// Update the move destination to wherever the cursor is pointing at
+	UpdateCachedDestination();
 }
 
 void ATP_TopDownPlayerController::OnSetDestinationTriggered()
@@ -71,23 +78,8 @@ void ATP_TopDownPlayerController::OnSetDestinationTriggered()
 	// We flag that the input is being pressed
 	FollowTime += GetWorld()->GetDeltaSeconds();
 	
-	// We look for the location in the world where the player has pressed the input
-	FHitResult Hit;
-	bool bHitSuccessful = false;
-	if (bIsTouch)
-	{
-		bHitSuccessful = GetHitResultUnderFinger(ETouchIndex::Touch1, ECollisionChannel::ECC_Visibility, true, Hit);
-	}
-	else
-	{
-		bHitSuccessful = GetHitResultUnderCursor(ECollisionChannel::ECC_Visibility, true, Hit);
-	}
-
-	// If we hit a surface, cache the location
-	if (bHitSuccessful)
-	{
-		CachedDestination = Hit.Location;
-	}
+	// Update the move destination to wherever the cursor is pointing at
+	UpdateCachedDestination();
 	
 	// Move towards mouse pointer or touch
 	APawn* ControlledPawn = GetPawn();
@@ -123,3 +115,24 @@ void ATP_TopDownPlayerController::OnTouchReleased()
 	bIsTouch = false;
 	OnSetDestinationReleased();
 }
+
+void ATP_TopDownPlayerController::UpdateCachedDestination()
+{
+	// We look for the location in the world where the player has pressed the input
+	FHitResult Hit;
+	bool bHitSuccessful = false;
+	if (bIsTouch)
+	{
+		bHitSuccessful = GetHitResultUnderFinger(ETouchIndex::Touch1, ECollisionChannel::ECC_Visibility, true, Hit);
+	}
+	else
+	{
+		bHitSuccessful = GetHitResultUnderCursor(ECollisionChannel::ECC_Visibility, true, Hit);
+	}
+
+	// If we hit a surface, cache the location
+	if (bHitSuccessful)
+	{
+		CachedDestination = Hit.Location;
+	}
+}
diff --git a/Templates/TP_TopDown/Source/TP_TopDown/TP_TopDownPlayerController.h b/Templates/TP_TopDown/Source/TP_TopDown/TP_TopDownPlayerController.h
--- a/Templates/TP_TopDown/Source/TP_TopDown/TP_TopDownPlayerController.h
+++ b/Templates/TP_TopDown/Source/TP_TopDown/TP_TopDownPlayerController.h
@@ -3,13 +3,14 @@
 #pragma once
 
 #include "CoreMinimal.h"
-#include "Templates/SubclassOf.h"
+//#include "Templates/SubclassOf.h"
 #include "GameFramework/PlayerController.h"
 #include "TP_TopDownPlayerController.generated.h"
 
 class UNiagaraSystem;
 class UInputMappingContext;
 class UInputAction;
+class UPathFollowingComponent;
 
 DECLARE_LOG_CATEGORY_EXTERN(LogTemplateCharacter, Log, All);
 
@@ -24,25 +25,29 @@ class ATP_TopDownPlayerController : public APlayerController
 
 protected:
 
+	/** Component used for moving along a NavMesh path. */
+	UPROPERTY(VisibleDefaultsOnly, Category = AI)
+	TObjectPtr<UPathFollowingComponent> PathFollowingComponent;
+
 	/** Time Threshold to know if it was a short press */
 	UPROPERTY(EditAnywhere, Category="Input")
 	float ShortPressThreshold;
 
 	/** FX Class that we will spawn when clicking */
 	UPROPERTY(EditAnywhere, Category="Input")
-	UNiagaraSystem* FXCursor;
+	TObjectPtr<UNiagaraSystem> FXCursor;
 
 	/** MappingContext */
 	UPROPERTY(EditAnywhere, Category="Input")
-	UInputMappingContext* DefaultMappingContext;
+	TObjectPtr<UInputMappingContext> DefaultMappingContext;
 	
 	/** Jump Input Action */
 	UPROPERTY(EditAnywhere, Category="Input")
-	UInputAction* SetDestinationClickAction;
+	TObjectPtr<UInputAction> SetDestinationClickAction;
 
 	/** Jump Input Action */
 	UPROPERTY(EditAnywhere, Category="Input")
-	UInputAction* SetDestinationTouchAction;
+	TObjectPtr<UInputAction> SetDestinationTouchAction;
 
 	/** True if the controlled character should navigate to the mouse cursor. */
 	uint32 bMoveToMouseCursor : 1;
@@ -73,6 +78,8 @@ protected:
 	void OnTouchTriggered();
 	void OnTouchReleased();
 
+	/** Helper function to get the move destination */
+	void UpdateCachedDestination();
 };
 
 
diff --git a/Templates/TP_TopDown/Source/TP_TopDown/Variant_TwinStick/Gameplay/TwinStickAoEAttack.cpp b/Templates/TP_TopDown/Source/TP_TopDown/Variant_TwinStick/Gameplay/TwinStickAoEAttack.cpp
--- a/Templates/TP_TopDown/Source/TP_TopDown/Variant_TwinStick/Gameplay/TwinStickAoEAttack.cpp
+++ b/Templates/TP_TopDown/Source/TP_TopDown/Variant_TwinStick/Gameplay/TwinStickAoEAttack.cpp
@@ -32,6 +32,8 @@ ATwinStickAoEAttack::ATwinStickAoEAttack()
 	CollisionSphere->SetCollisionObjectType(ECC_WorldDynamic);
 	CollisionSphere->SetCollisionResponseToAllChannels(ECR_Ignore);
 	CollisionSphere->SetCollisionResponseToChannel(ECC_Pawn, ECR_Overlap);
+
+	CollisionSphere->OnComponentBeginOverlap.AddDynamic(this, &ATwinStickAoEAttack::OnAoEOverlap);
 }
 
 void ATwinStickAoEAttack::BeginPlay()
@@ -39,7 +41,7 @@ void ATwinStickAoEAttack::BeginPlay()
 	Super::BeginPlay();
 	
 	// set up the AoE timers
-	GetWorld()->GetTimerManager().SetTimer(TickAoETimer, this, &ATwinStickAoEAttack::TickAoE, TickAoETime, true);
+	GetWorld()->GetTimerManager().SetTimer(StartAoETimer, this, &ATwinStickAoEAttack::StartAoE, StartAoETime, false);
 	GetWorld()->GetTimerManager().SetTimer(StopAoETimer, this, &ATwinStickAoEAttack::StopAoE, StopAoETime, false);
 
 }
@@ -49,12 +51,15 @@ void ATwinStickAoEAttack::EndPlay(EEndPlayReason::Type EndPlayReason)
 	Super::EndPlay(EndPlayReason);
 
 	// clear the timers
-	GetWorld()->GetTimerManager().ClearTimer(TickAoETimer);
+	GetWorld()->GetTimerManager().ClearTimer(StartAoETimer);
 	GetWorld()->GetTimerManager().ClearTimer(StopAoETimer);
 }
 
-void ATwinStickAoEAttack::TickAoE()
+void ATwinStickAoEAttack::StartAoE()
 {
+	// raise the active flag
+	bIsAoEActive = true;
+
 	// find all actors overlapping the NPC
 	TArray<AActor*> Overlaps;
 	CollisionSphere->GetOverlappingActors(Overlaps, ATwinStickNPC::StaticClass());
@@ -72,12 +77,26 @@ void ATwinStickAoEAttack::TickAoE()
 
 void ATwinStickAoEAttack::StopAoE()
 {
-	// stop the damage tick timer
-	GetWorld()->GetTimerManager().ClearTimer(TickAoETimer);
+	// drop the active flag
+	bIsAoEActive = false;
 
-	// hide the mesh
-	SphereVisual->SetHiddenInGame(true);
+	// stop the damage tick timer
+	GetWorld()->GetTimerManager().ClearTimer(StartAoETimer);
 
 	// call the BP handler. It will be responsible for destroying the Actor when it's done
 	BP_AoEFinished();
 }
+
+void ATwinStickAoEAttack::OnAoEOverlap(UPrimitiveComponent* OverlappedComponent, AActor* OtherActor, UPrimitiveComponent* OtherComp, int32 OtherBodyIndex, bool bFromSweep, const FHitResult& SweepResult)
+{
+	// is the explosion active?
+	if (bIsAoEActive)
+	{
+		// did we overlap an NPC?
+		if (ATwinStickNPC* NPC = Cast<ATwinStickNPC>(OtherActor))
+		{
+			// tell the NPC it's been hit
+			NPC->ProjectileImpact(FVector::ZeroVector);
+		}
+	}
+}
diff --git a/Templates/TP_TopDown/Source/TP_TopDown/Variant_TwinStick/Gameplay/TwinStickAoEAttack.h b/Templates/TP_TopDown/Source/TP_TopDown/Variant_TwinStick/Gameplay/TwinStickAoEAttack.h
--- a/Templates/TP_TopDown/Source/TP_TopDown/Variant_TwinStick/Gameplay/TwinStickAoEAttack.h
+++ b/Templates/TP_TopDown/Source/TP_TopDown/Variant_TwinStick/Gameplay/TwinStickAoEAttack.h
@@ -29,18 +29,21 @@ class ATwinStickAoEAttack : public AActor
 protected:
 
 	/** Timer to start AoE damage checks */
-	FTimerHandle TickAoETimer;
+	FTimerHandle StartAoETimer;
 
 	/** Timer to end AoE damage checks */
 	FTimerHandle StopAoETimer;
 
 	/** Time to wait between AoE damage ticks */
 	UPROPERTY(EditAnywhere, Category="AoE Attack", meta=(ClampMin = 0, ClampMax = 5, Units = "s"))
-	float TickAoETime = 0.33f;
+	float StartAoETime = 0.033f;
 
 	/** Time to wait before stopping AoE damage checks */
 	UPROPERTY(EditAnywhere, Category="AoE Attack", meta=(ClampMin = 0, ClampMax = 5, Units = "s"))
-	float StopAoETime = 1.0f;
+	float StopAoETime = 0.5f;
+
+	/** While true, the AoE will damage anything that overlaps it */
+	bool bIsAoEActive = false;
 
 public:	
 	
@@ -58,7 +61,7 @@ protected:
 protected:
 
 	/** Called when the start AoE timer triggers */
-	void TickAoE();
+	void StartAoE();
 
 	/** Called when the stop AoE timer triggers */
 	void StopAoE();
@@ -66,4 +69,8 @@ protected:
 	/** Allows Blueprint handling of AoE fade out effects. NOTE: Call Destroy Actor at the end of this! */
 	UFUNCTION(BlueprintImplementableEvent, Category="AoE Attack")
 	void BP_AoEFinished();
+
+	/** Handles collision with the AoE while it's active */
+	UFUNCTION()
+	void OnAoEOverlap(UPrimitiveComponent* OverlappedComponent, AActor* OtherActor, UPrimitiveComponent* OtherComp, int32 OtherBodyIndex, bool bFromSweep, const FHitResult& SweepResult);
 };
diff --git a/Templates/TP_TopDown/Source/TP_TopDown/Variant_TwinStick/TwinStickGameMode.cpp b/Templates/TP_TopDown/Source/TP_TopDown/Variant_TwinStick/TwinStickGameMode.cpp
--- a/Templates/TP_TopDown/Source/TP_TopDown/Variant_TwinStick/TwinStickGameMode.cpp
+++ b/Templates/TP_TopDown/Source/TP_TopDown/Variant_TwinStick/TwinStickGameMode.cpp
@@ -9,9 +9,8 @@
 
 void ATwinStickGameMode::BeginPlay()
 {
-	// create the UI widget and add it to the viewport
-	UIWidget = CreateWidget<UTwinStickUI>(UGameplayStatics::GetPlayerController(GetWorld(), 0), UIWidgetClass);
-	UIWidget->AddToViewport(0);
+	// create the UI widget if it hasn't already
+	CreateUI();
 }
 
 void ATwinStickGameMode::EndPlay(EEndPlayReason::Type EndPlayReason)
@@ -24,6 +23,12 @@ void ATwinStickGameMode::EndPlay(EEndPlayReason::Type EndPlayReason)
 
 void ATwinStickGameMode::ItemUsed(int32 Value)
 {
+	// ensure the UI widget is available
+	if (!UIWidget)
+	{
+		CreateUI();
+	}
+
 	// update the UI
 	UIWidget->UpdateItems(Value);
 }
@@ -40,6 +45,17 @@ void ATwinStickGameMode::ScoreUpdate(int32 Value)
 	ComboUpdate();
 }
 
+void ATwinStickGameMode::CreateUI()
+{
+	// avoid creating the UI multiple times
+	if(UIWidget)
+		return;
+
+	// create the UI widget and add it to the viewport
+	UIWidget = CreateWidget<UTwinStickUI>(UGameplayStatics::GetPlayerController(GetWorld(), 0), UIWidgetClass);
+	UIWidget->AddToViewport(0);
+}
+
 void ATwinStickGameMode::ComboUpdate()
 {
 	// return
diff --git a/Templates/TP_TopDown/Source/TP_TopDown/Variant_TwinStick/TwinStickGameMode.h b/Templates/TP_TopDown/Source/TP_TopDown/Variant_TwinStick/TwinStickGameMode.h
--- a/Templates/TP_TopDown/Source/TP_TopDown/Variant_TwinStick/TwinStickGameMode.h
+++ b/Templates/TP_TopDown/Source/TP_TopDown/Variant_TwinStick/TwinStickGameMode.h
@@ -77,6 +77,9 @@ public:
 
 protected:
 
+	/** Creates the UI widget if it hasn't been created already */
+	void CreateUI();
+
 	/** Updates the combo multiplier */
 	void ComboUpdate();
 
