TP_ThirdPerson, TP_TopDown, TP_VehicleAdv and BP versions
- Implemented force touch controls flag on all relevant PlayerControllers
- Split mobile-conflicting inputs into a Mobile Excluded asset on TwinStick variant.
#rnx

#virtualized

diff --git a/Templates/TP_ThirdPerson/Source/TP_ThirdPerson/TP_ThirdPersonPlayerController.cpp b/Templates/TP_ThirdPerson/Source/TP_ThirdPerson/TP_ThirdPersonPlayerController.cpp
--- a/Templates/TP_ThirdPerson/Source/TP_ThirdPerson/TP_ThirdPersonPlayerController.cpp
+++ b/Templates/TP_ThirdPerson/Source/TP_ThirdPerson/TP_ThirdPersonPlayerController.cpp
@@ -14,7 +14,7 @@ void ATP_ThirdPersonPlayerController::BeginPlay()
 	Super::BeginPlay();
 
 	// only spawn touch controls on local player controllers
-	if (SVirtualJoystick::ShouldDisplayTouchInterface() && IsLocalPlayerController())
+	if (ShouldUseTouchControls() && IsLocalPlayerController())
 	{
 		// spawn the mobile controls widget
 		MobileControlsWidget = CreateWidget<UUserWidget>(this, MobileControlsWidgetClass);
@@ -49,7 +49,7 @@ void ATP_ThirdPersonPlayerController::SetupInputComponent()
 			}
 
 			// only add these IMCs if we're not using mobile touch input
-			if (!SVirtualJoystick::ShouldDisplayTouchInterface())
+			if (!ShouldUseTouchControls())
 			{
 				for (UInputMappingContext* CurrentContext : MobileExcludedMappingContexts)
 				{
@@ -59,3 +59,9 @@ void ATP_ThirdPersonPlayerController::SetupInputComponent()
 		}
 	}
 }
+
+bool ATP_ThirdPersonPlayerController::ShouldUseTouchControls() const
+{
+	// are we on a mobile platform? Should we force touch?
+	return SVirtualJoystick::ShouldDisplayTouchInterface() || bForceTouchControls;
+}
diff --git a/Templates/TP_ThirdPerson/Source/TP_ThirdPerson/TP_ThirdPersonPlayerController.h b/Templates/TP_ThirdPerson/Source/TP_ThirdPerson/TP_ThirdPersonPlayerController.h
--- a/Templates/TP_ThirdPerson/Source/TP_ThirdPerson/TP_ThirdPersonPlayerController.h
+++ b/Templates/TP_ThirdPerson/Source/TP_ThirdPerson/TP_ThirdPersonPlayerController.h
@@ -35,10 +35,17 @@ protected:
 	/** Pointer to the mobile controls widget */
 	TObjectPtr<UUserWidget> MobileControlsWidget;
 
+	/** If true, the player will use UMG touch controls even if not playing on mobile platforms */
+	UPROPERTY(EditAnywhere, Config, Category = "Input|Touch Controls")
+	bool bForceTouchControls = false;
+
 	/** Gameplay initialization */
 	virtual void BeginPlay() override;
 
 	/** Input mapping context setup */
 	virtual void SetupInputComponent() override;
 
+	/** Returns true if the player should use UMG touch controls */
+	bool ShouldUseTouchControls() const;
+
 };
diff --git a/Templates/TP_ThirdPerson/Source/TP_ThirdPerson/Variant_Combat/CombatPlayerController.cpp b/Templates/TP_ThirdPerson/Source/TP_ThirdPerson/Variant_Combat/CombatPlayerController.cpp
--- a/Templates/TP_ThirdPerson/Source/TP_ThirdPerson/Variant_Combat/CombatPlayerController.cpp
+++ b/Templates/TP_ThirdPerson/Source/TP_ThirdPerson/Variant_Combat/CombatPlayerController.cpp
@@ -18,7 +18,7 @@ void ACombatPlayerController::BeginPlay()
 	Super::BeginPlay();
 
 	// only spawn touch controls on local player controllers
-	if (SVirtualJoystick::ShouldDisplayTouchInterface() && IsLocalPlayerController())
+	if (ShouldUseTouchControls() && IsLocalPlayerController())
 	{
 		// spawn the mobile controls widget
 		MobileControlsWidget = CreateWidget<UUserWidget>(this, MobileControlsWidgetClass);
@@ -51,7 +51,7 @@ void ACombatPlayerController::SetupInputComponent()
 			}
 
 			// only add these IMCs if we're not using mobile touch input
-			if (!SVirtualJoystick::ShouldDisplayTouchInterface())
+			if (!ShouldUseTouchControls())
 			{
 				for (UInputMappingContext* CurrentContext : MobileExcludedMappingContexts)
 				{
@@ -84,4 +84,10 @@ void ACombatPlayerController::OnPawnDestroyed(AActor* DestroyedActor)
 		// possess the character
 		Possess(RespawnedCharacter);
 	}
-}
\ No newline at end of file
+}
+
+bool ACombatPlayerController::ShouldUseTouchControls() const
+{
+	// are we on a mobile platform? Should we force touch?
+	return SVirtualJoystick::ShouldDisplayTouchInterface() || bForceTouchControls;
+}
diff --git a/Templates/TP_ThirdPerson/Source/TP_ThirdPerson/Variant_Combat/CombatPlayerController.h b/Templates/TP_ThirdPerson/Source/TP_ThirdPerson/Variant_Combat/CombatPlayerController.h
--- a/Templates/TP_ThirdPerson/Source/TP_ThirdPerson/Variant_Combat/CombatPlayerController.h
+++ b/Templates/TP_ThirdPerson/Source/TP_ThirdPerson/Variant_Combat/CombatPlayerController.h
@@ -36,6 +36,10 @@ protected:
 	/** Pointer to the mobile controls widget */
 	TObjectPtr<UUserWidget> MobileControlsWidget;
 
+	/** If true, the player will use UMG touch controls even if not playing on mobile platforms */
+	UPROPERTY(EditAnywhere, Config, Category = "Input|Touch Controls")
+	bool bForceTouchControls = false;
+
 	/** Character class to respawn when the possessed pawn is destroyed */
 	UPROPERTY(EditAnywhere, Category="Respawn")
 	TSubclassOf<ACombatCharacter> CharacterClass;
@@ -65,4 +69,7 @@ protected:
 	UFUNCTION()
 	void OnPawnDestroyed(AActor* DestroyedActor);
 
+	/** Returns true if the player should use UMG touch controls */
+	bool ShouldUseTouchControls() const;
+
 };
diff --git a/Templates/TP_ThirdPerson/Source/TP_ThirdPerson/Variant_Platforming/PlatformingPlayerController.cpp b/Templates/TP_ThirdPerson/Source/TP_ThirdPerson/Variant_Platforming/PlatformingPlayerController.cpp
--- a/Templates/TP_ThirdPerson/Source/TP_ThirdPerson/Variant_Platforming/PlatformingPlayerController.cpp
+++ b/Templates/TP_ThirdPerson/Source/TP_ThirdPerson/Variant_Platforming/PlatformingPlayerController.cpp
@@ -18,7 +18,7 @@ void APlatformingPlayerController::BeginPlay()
 	Super::BeginPlay();
 
 	// only spawn touch controls on local player controllers
-	if (SVirtualJoystick::ShouldDisplayTouchInterface() && IsLocalPlayerController())
+	if (ShouldUseTouchControls() && IsLocalPlayerController())
 	{
 		// spawn the mobile controls widget
 		MobileControlsWidget = CreateWidget<UUserWidget>(this, MobileControlsWidgetClass);
@@ -51,7 +51,7 @@ void APlatformingPlayerController::SetupInputComponent()
 			}
 
 			// only add these IMCs if we're not using mobile touch input
-			if (!SVirtualJoystick::ShouldDisplayTouchInterface())
+			if (!ShouldUseTouchControls())
 			{
 				for (UInputMappingContext* CurrentContext : MobileExcludedMappingContexts)
 				{
@@ -87,4 +87,10 @@ void APlatformingPlayerController::OnPawnDestroyed(AActor* DestroyedActor)
 			Possess(RespawnedCharacter);
 		}
 	}
-}
\ No newline at end of file
+}
+
+bool APlatformingPlayerController::ShouldUseTouchControls() const
+{
+	// are we on a mobile platform? Should we force touch?
+	return SVirtualJoystick::ShouldDisplayTouchInterface() || bForceTouchControls;
+}
diff --git a/Templates/TP_ThirdPerson/Source/TP_ThirdPerson/Variant_Platforming/PlatformingPlayerController.h b/Templates/TP_ThirdPerson/Source/TP_ThirdPerson/Variant_Platforming/PlatformingPlayerController.h
--- a/Templates/TP_ThirdPerson/Source/TP_ThirdPerson/Variant_Platforming/PlatformingPlayerController.h
+++ b/Templates/TP_ThirdPerson/Source/TP_ThirdPerson/Variant_Platforming/PlatformingPlayerController.h
@@ -36,6 +36,10 @@ protected:
 	/** Pointer to the mobile controls widget */
 	TObjectPtr<UUserWidget> MobileControlsWidget;
 
+	/** If true, the player will use UMG touch controls even if not playing on mobile platforms */
+	UPROPERTY(EditAnywhere, Config, Category = "Input|Touch Controls")
+	bool bForceTouchControls = false;
+
 	/** Character class to respawn when the possessed pawn is destroyed */
 	UPROPERTY(EditAnywhere, Category="Respawn")
 	TSubclassOf<APlatformingCharacter> CharacterClass;
@@ -54,4 +58,7 @@ protected:
 	/** Called if the possessed pawn is destroyed */
 	UFUNCTION()
 	void OnPawnDestroyed(AActor* DestroyedActor);
+
+	/** Returns true if the player should use UMG touch controls */
+	bool ShouldUseTouchControls() const;
 };
diff --git a/Templates/TP_ThirdPerson/Source/TP_ThirdPerson/Variant_SideScrolling/SideScrollingPlayerController.cpp b/Templates/TP_ThirdPerson/Source/TP_ThirdPerson/Variant_SideScrolling/SideScrollingPlayerController.cpp
--- a/Templates/TP_ThirdPerson/Source/TP_ThirdPerson/Variant_SideScrolling/SideScrollingPlayerController.cpp
+++ b/Templates/TP_ThirdPerson/Source/TP_ThirdPerson/Variant_SideScrolling/SideScrollingPlayerController.cpp
@@ -18,7 +18,7 @@ void ASideScrollingPlayerController::BeginPlay()
 	Super::BeginPlay();
 
 	// only spawn touch controls on local player controllers
-	if (SVirtualJoystick::ShouldDisplayTouchInterface() && IsLocalPlayerController())
+	if (ShouldUseTouchControls() && IsLocalPlayerController())
 	{
 		// spawn the mobile controls widget
 		MobileControlsWidget = CreateWidget<UUserWidget>(this, MobileControlsWidgetClass);
@@ -51,7 +51,7 @@ void ASideScrollingPlayerController::SetupInputComponent()
 			}
 
 			// only add these IMCs if we're not using mobile touch input
-			if (!SVirtualJoystick::ShouldDisplayTouchInterface())
+			if (!ShouldUseTouchControls())
 			{
 				for (UInputMappingContext* CurrentContext : MobileExcludedMappingContexts)
 				{
@@ -88,3 +88,9 @@ void ASideScrollingPlayerController::OnPawnDestroyed(AActor* DestroyedActor)
 		}
 	}
 }
+
+bool ASideScrollingPlayerController::ShouldUseTouchControls() const
+{
+	// are we on a mobile platform? Should we force touch?
+	return SVirtualJoystick::ShouldDisplayTouchInterface() || bForceTouchControls;
+}
diff --git a/Templates/TP_ThirdPerson/Source/TP_ThirdPerson/Variant_SideScrolling/SideScrollingPlayerController.h b/Templates/TP_ThirdPerson/Source/TP_ThirdPerson/Variant_SideScrolling/SideScrollingPlayerController.h
--- a/Templates/TP_ThirdPerson/Source/TP_ThirdPerson/Variant_SideScrolling/SideScrollingPlayerController.h
+++ b/Templates/TP_ThirdPerson/Source/TP_ThirdPerson/Variant_SideScrolling/SideScrollingPlayerController.h
@@ -37,6 +37,10 @@ protected:
 	/** Pointer to the mobile controls widget */
 	TObjectPtr<UUserWidget> MobileControlsWidget;
 
+	/** If true, the player will use UMG touch controls even if not playing on mobile platforms */
+	UPROPERTY(EditAnywhere, Config, Category = "Input|Touch Controls")
+	bool bForceTouchControls = false;
+
 	/** Character class to respawn when the possessed pawn is destroyed */
 	UPROPERTY(EditAnywhere, Category="Respawn")
 	TSubclassOf<ASideScrollingCharacter> CharacterClass;
@@ -56,4 +60,7 @@ protected:
 	UFUNCTION()
 	void OnPawnDestroyed(AActor* DestroyedActor);
 
+	/** Returns true if the player should use UMG touch controls */
+	bool ShouldUseTouchControls() const;
+
 };
diff --git a/Templates/TP_TopDown/Source/TP_TopDown/Variant_TwinStick/TwinStickPlayerController.cpp b/Templates/TP_TopDown/Source/TP_TopDown/Variant_TwinStick/TwinStickPlayerController.cpp
--- a/Templates/TP_TopDown/Source/TP_TopDown/Variant_TwinStick/TwinStickPlayerController.cpp
+++ b/Templates/TP_TopDown/Source/TP_TopDown/Variant_TwinStick/TwinStickPlayerController.cpp
@@ -18,7 +18,7 @@ void ATwinStickPlayerController::BeginPlay()
 	Super::BeginPlay();
 
 	// only spawn touch controls on local player controllers
-	if (SVirtualJoystick::ShouldDisplayTouchInterface() && IsLocalPlayerController())
+	if (ShouldUseTouchControls() && IsLocalPlayerController())
 	{
 		// spawn the mobile controls widget
 		MobileControlsWidget = CreateWidget<UUserWidget>(this, MobileControlsWidgetClass);
@@ -51,6 +51,15 @@ void ATwinStickPlayerController::SetupInputComponent()
 			{
 				Subsystem->AddMappingContext(CurrentContext, 0);
 			}
+
+			// only add these IMCs if we're not using mobile touch input
+			if (!ShouldUseTouchControls())
+			{
+				for (UInputMappingContext* CurrentContext : MobileExcludedMappingContexts)
+				{
+					Subsystem->AddMappingContext(CurrentContext, 0);
+				}
+			}
 		}
 	}
 }
@@ -81,3 +90,9 @@ void ATwinStickPlayerController::OnPawnDestroyed(AActor* DestroyedActor)
 		}
 	}
 }
+
+bool ATwinStickPlayerController::ShouldUseTouchControls() const
+{
+	// are we on a mobile platform? Should we force touch?
+	return SVirtualJoystick::ShouldDisplayTouchInterface() || bForceTouchControls;
+}
diff --git a/Templates/TP_TopDown/Source/TP_TopDown/Variant_TwinStick/TwinStickPlayerController.h b/Templates/TP_TopDown/Source/TP_TopDown/Variant_TwinStick/TwinStickPlayerController.h
--- a/Templates/TP_TopDown/Source/TP_TopDown/Variant_TwinStick/TwinStickPlayerController.h
+++ b/Templates/TP_TopDown/Source/TP_TopDown/Variant_TwinStick/TwinStickPlayerController.h
@@ -25,6 +25,10 @@ protected:
 	UPROPERTY(EditAnywhere, Category ="Input|Input Mappings")
 	TArray<UInputMappingContext*> DefaultMappingContexts;
 
+	/** Input Mapping Contexts */
+	UPROPERTY(EditAnywhere, Category="Input|Input Mappings")
+	TArray<UInputMappingContext*> MobileExcludedMappingContexts;
+
 	/** Mobile controls widget to spawn */
 	UPROPERTY(EditAnywhere, Category="Input|Touch Controls")
 	TSubclassOf<UUserWidget> MobileControlsWidgetClass;
@@ -32,6 +36,10 @@ protected:
 	/** Pointer to the mobile controls widget */
 	TObjectPtr<UUserWidget> MobileControlsWidget;
 
+	/** If true, the player will use UMG touch controls even if not playing on mobile platforms */
+	UPROPERTY(EditAnywhere, Config, Category = "Input|Touch Controls")
+	bool bForceTouchControls = false;
+
 	/** Character class to respawn when the possessed pawn is destroyed */
 	UPROPERTY(EditAnywhere, Category="Respawn")
 	TSubclassOf<ATwinStickCharacter> CharacterClass;
@@ -50,4 +58,7 @@ protected:
 	/** Called if the possessed pawn is destroyed */
 	UFUNCTION()
 	void OnPawnDestroyed(AActor* DestroyedActor);
+
+	/** Returns true if the player should use UMG touch controls */
+	bool ShouldUseTouchControls() const;
 };
diff --git a/Templates/TP_VehicleAdv/Source/TP_VehicleAdv/TP_VehicleAdvPlayerController.cpp b/Templates/TP_VehicleAdv/Source/TP_VehicleAdv/TP_VehicleAdvPlayerController.cpp
--- a/Templates/TP_VehicleAdv/Source/TP_VehicleAdv/TP_VehicleAdvPlayerController.cpp
+++ b/Templates/TP_VehicleAdv/Source/TP_VehicleAdv/TP_VehicleAdvPlayerController.cpp
@@ -22,7 +22,7 @@ void ATP_VehicleAdvPlayerController::BeginPlay()
 	// only spawn UI on local player controllers
 	if (IsLocalPlayerController())
 	{
-		if (SVirtualJoystick::ShouldDisplayTouchInterface())
+		if (ShouldUseTouchControls())
 		{
 			// spawn the mobile controls widget
 			MobileControlsWidget = CreateWidget<UUserWidget>(this, MobileControlsWidgetClass);
@@ -71,7 +71,7 @@ void ATP_VehicleAdvPlayerController::SetupInputComponent()
 			}
 
 			// only add these IMCs if we're not using mobile touch input
-			if (!SVirtualJoystick::ShouldDisplayTouchInterface())
+			if (!ShouldUseTouchControls())
 			{
 				for (UInputMappingContext* CurrentContext : MobileExcludedMappingContexts)
 				{
@@ -122,3 +122,9 @@ void ATP_VehicleAdvPlayerController::OnPawnDestroyed(AActor* DestroyedPawn)
 		}
 	}
 }
+
+bool ATP_VehicleAdvPlayerController::ShouldUseTouchControls() const
+{
+	// are we on a mobile platform? Should we force touch?
+	return SVirtualJoystick::ShouldDisplayTouchInterface() || bForceTouchControls;
+}
diff --git a/Templates/TP_VehicleAdv/Source/TP_VehicleAdv/TP_VehicleAdvPlayerController.h b/Templates/TP_VehicleAdv/Source/TP_VehicleAdv/TP_VehicleAdvPlayerController.h
--- a/Templates/TP_VehicleAdv/Source/TP_VehicleAdv/TP_VehicleAdvPlayerController.h
+++ b/Templates/TP_VehicleAdv/Source/TP_VehicleAdv/TP_VehicleAdvPlayerController.h
@@ -36,6 +36,10 @@ protected:
 	/** Pointer to the mobile controls widget */
 	TObjectPtr<UUserWidget> MobileControlsWidget;
 
+	/** If true, the player will use UMG touch controls even if not playing on mobile platforms */
+	UPROPERTY(EditAnywhere, Config, Category = "Input|Touch Controls")
+	bool bForceTouchControls = false;
+
 	/** If true, the optional steering wheel input mapping context will be registered */
 	UPROPERTY(EditAnywhere, Category = "Input|Steering Wheel Controls")
 	bool bUseSteeringWheelControls = false;
@@ -81,4 +85,7 @@ protected:
 	/** Handles pawn destruction and respawning */
 	UFUNCTION()
 	void OnPawnDestroyed(AActor* DestroyedPawn);
+
+	/** Returns true if the player should use UMG touch controls */
+	bool ShouldUseTouchControls() const;
 };
diff --git a/Templates/TP_VehicleAdv/Source/TP_VehicleAdv/Variant_TimeTrial/TimeTrialPlayerController.cpp b/Templates/TP_VehicleAdv/Source/TP_VehicleAdv/Variant_TimeTrial/TimeTrialPlayerController.cpp
--- a/Templates/TP_VehicleAdv/Source/TP_VehicleAdv/Variant_TimeTrial/TimeTrialPlayerController.cpp
+++ b/Templates/TP_VehicleAdv/Source/TP_VehicleAdv/Variant_TimeTrial/TimeTrialPlayerController.cpp
@@ -25,7 +25,7 @@ void ATimeTrialPlayerController::BeginPlay()
 	// only spawn UI on local player controllers
 	if (IsLocalPlayerController())
 	{
-		if (SVirtualJoystick::ShouldDisplayTouchInterface())
+		if (ShouldUseTouchControls())
 		{
 			// spawn the mobile controls widget
 			MobileControlsWidget = CreateWidget<UUserWidget>(this, MobileControlsWidgetClass);
@@ -91,7 +91,7 @@ void ATimeTrialPlayerController::SetupInputComponent()
 			}
 
 			// only add these IMCs if we're not using mobile touch input
-			if (!SVirtualJoystick::ShouldDisplayTouchInterface())
+			if (!ShouldUseTouchControls())
 			{
 				for (UInputMappingContext* CurrentContext : MobileExcludedMappingContexts)
 				{
@@ -186,3 +186,9 @@ void ATimeTrialPlayerController::OnPawnDestroyed(AActor* DestroyedPawn)
 		}
 	}
 }
+
+bool ATimeTrialPlayerController::ShouldUseTouchControls() const
+{
+	// are we on a mobile platform? Should we force touch?
+	return SVirtualJoystick::ShouldDisplayTouchInterface() || bForceTouchControls;
+}
diff --git a/Templates/TP_VehicleAdv/Source/TP_VehicleAdv/Variant_TimeTrial/TimeTrialPlayerController.h b/Templates/TP_VehicleAdv/Source/TP_VehicleAdv/Variant_TimeTrial/TimeTrialPlayerController.h
--- a/Templates/TP_VehicleAdv/Source/TP_VehicleAdv/Variant_TimeTrial/TimeTrialPlayerController.h
+++ b/Templates/TP_VehicleAdv/Source/TP_VehicleAdv/Variant_TimeTrial/TimeTrialPlayerController.h
@@ -37,6 +37,10 @@ protected:
 	/** Pointer to the mobile controls widget */
 	TObjectPtr<UUserWidget> MobileControlsWidget;
 
+	/** If true, the player will use UMG touch controls even if not playing on mobile platforms */
+	UPROPERTY(EditAnywhere, Config, Category = "Input|Touch Controls")
+	bool bForceTouchControls = false;
+
 	/** If true, the optional steering wheel input mapping context will be registered */
 	UPROPERTY(EditAnywhere, Category = "Input|Steering Wheel Controls")
 	bool bUseSteeringWheelControls = false;
@@ -113,4 +117,7 @@ protected:
 	/** Handles pawn destruction and respawning */
 	UFUNCTION()
 	void OnPawnDestroyed(AActor* DestroyedPawn);
+
+	/** Returns true if the player should use UMG touch controls */
+	bool ShouldUseTouchControls() const;
 };
