TP_TopDown, TP_TopDownBP
- Strategy: Fixed nonfunctional Touch controls on mobile
- Strategy: Removed overlapping zoom controls on mobile, pending an UI rework
#rnx

diff --git a/Templates/TP_TopDown/Source/TP_TopDown/Variant_Strategy/StrategyPlayerController.cpp b/Templates/TP_TopDown/Source/TP_TopDown/Variant_Strategy/StrategyPlayerController.cpp
--- a/Templates/TP_TopDown/Source/TP_TopDown/Variant_Strategy/StrategyPlayerController.cpp
+++ b/Templates/TP_TopDown/Source/TP_TopDown/Variant_Strategy/StrategyPlayerController.cpp
@@ -78,15 +78,13 @@ void AStrategyPlayerController::SetupInputComponent()
 			// Touch Interaction
 			EnhancedInputComponent->BindAction(TouchPrimaryHoldAction, ETriggerEvent::Started, this, &AStrategyPlayerController::TouchPrimaryHoldStarted);
 			EnhancedInputComponent->BindAction(TouchPrimaryHoldAction, ETriggerEvent::Triggered, this, &AStrategyPlayerController::TouchPrimaryHoldTriggered);
-
-			EnhancedInputComponent->BindAction(TouchPrimaryTapAction, ETriggerEvent::Completed, this, &AStrategyPlayerController::TouchPrimaryTap);
+			EnhancedInputComponent->BindAction(TouchPrimaryHoldAction, ETriggerEvent::Completed, this, &AStrategyPlayerController::TouchPrimaryHoldCompleted);
 
 			EnhancedInputComponent->BindAction(TouchSecondaryAction, ETriggerEvent::Started, this, &AStrategyPlayerController::TouchSecondaryStarted);
 			EnhancedInputComponent->BindAction(TouchSecondaryAction, ETriggerEvent::Triggered, this, &AStrategyPlayerController::TouchSecondaryTriggered);
 			EnhancedInputComponent->BindAction(TouchSecondaryAction, ETriggerEvent::Completed, this, &AStrategyPlayerController::TouchSecondaryCompleted);
 			EnhancedInputComponent->BindAction(TouchSecondaryAction, ETriggerEvent::Canceled, this, &AStrategyPlayerController::TouchSecondaryCompleted);
 
-			EnhancedInputComponent->BindAction(TouchDoubleTapAction, ETriggerEvent::Triggered, this, &AStrategyPlayerController::TouchDoubleTap);
 		}
 	}
 }
@@ -124,6 +122,17 @@ void AStrategyPlayerController::DragSelectUnits(const TArray<AStrategyUnit*>& Un
 			// select the unit
 			CurrentUnit->UnitSelected();
 		}
+
+	}
+	else
+	{
+
+		// release any currently selected units since nothing is on the box
+		if (ControlledUnits.Num() > 0)
+		{
+			DoDeselectAllCommand();
+		}
+
 	}
 }
 
@@ -193,14 +202,20 @@ void AStrategyPlayerController::SelectHoldTriggered(const FInputActionValue& Val
 	FVector2D SelectionSize = SelectionPosition - StartingSelectionPosition;
 
 	// update the selection box on the HUD
-	StrategyHUD->DragSelectUpdate(StartingSelectionPosition, SelectionSize, SelectionPosition, true);
+	if (StrategyHUD)
+	{
+		StrategyHUD->DragSelectUpdate(StartingSelectionPosition, SelectionSize, SelectionPosition, true);
+	}
+	
 }
 
 void AStrategyPlayerController::SelectHoldCompleted(const FInputActionValue& Value)
 {
-
 	// reset the drag box on the HUD
-	StrategyHUD->DragSelectUpdate(FVector2D::ZeroVector, FVector2D::ZeroVector, FVector2D::ZeroVector, false);
+	if (StrategyHUD)
+	{
+		StrategyHUD->DragSelectUpdate(FVector2D::ZeroVector, FVector2D::ZeroVector, FVector2D::ZeroVector, false);
+	}
 }
 
 void AStrategyPlayerController::SelectClick(const FInputActionValue& Value)
@@ -253,7 +268,8 @@ void AStrategyPlayerController::InteractClickCompleted(const FInputActionValue&
 			bDoubleTapActive = false;
 
 		}
-		else {
+		else
+		{
 
 			// move the selected units to the target location
 			DoMoveUnitsCommand();
@@ -264,6 +280,8 @@ void AStrategyPlayerController::InteractClickCompleted(const FInputActionValue&
 
 void AStrategyPlayerController::TouchPrimaryHoldStarted(const FInputActionValue& Value)
 {
+	// save the tap press time
+	LastTapPressTime = GetWorld()->GetRealTimeSeconds();
 
 	// save the starting interaction position
 	StartingInteractionPosition = Value.Get<FVector2D>();
@@ -271,32 +289,66 @@ void AStrategyPlayerController::TouchPrimaryHoldStarted(const FInputActionValue&
 
 void AStrategyPlayerController::TouchPrimaryHoldTriggered(const FInputActionValue& Value)
 {
-
-	// are we in selection modifier mode, and have a large enough selection box?
-	if (bSelectionModifier && StartingSecondFingerPosition.Equals(CurrentSecondFingerPosition, 10.0f))
+	// is this touch longer than a tap?
+	if ((GetWorld()->GetRealTimeSeconds() - LastTapPressTime) > TouchTapMaxAllowedTime)
 	{
-		// update the interaction position
-		CurrentInteractionPosition = Value.Get<FVector2D>();
+		// if we're not doing a box select, do a drag scroll
+		if (!bSelectionModifier)
+		{
+			DoDragScrollCommand();
+		}
+	}
+}
 
-		// update the selection box on the HUD
-		StrategyHUD->DragSelectUpdate(StartingInteractionPosition, CurrentInteractionPosition - StartingSecondFingerPosition, CurrentInteractionPosition, true);
+void AStrategyPlayerController::TouchPrimaryHoldCompleted(const FInputActionValue& Value)
+{
+	// check if we're doing a tap or double tap.
+	// we have to do this manually because EnhancedInput tap triggers work differently on touch inputs
+	bool bTapped = false;
+	bool bDoubleTapped = false;
 
-	} else {
+	CheckTouchTap(bTapped, bDoubleTapped);
 
-		// do a drag scroll instead
-		DoDragScrollCommand();
+	// do we have a double tap?
+	if (bTapped)
+	{
+		if (bDoubleTapped)
+		{
+			// ensure are we not doing a box select
+			if (!bSelectionModifier)
+			{
+				// depending on the double tap toggle, select or deselect all units
+				if (bDoubleTapActive)
+				{
+					DoDeselectAllCommand();
+				}
+				else
+				{
+					DoSelectAllOnScreenCommand();
+				}
 
+				// toggle the double tap flag
+				bDoubleTapActive = !bDoubleTapActive;
+			}
+		}
+
+	// no double tap, handle this touch input normally
 	}
-}
+	else
+	{
 
-void AStrategyPlayerController::TouchPrimaryTap(const FInputActionValue& Value)
-{
+		// ensure we're not already box selecting, or were just box selecting
+		if (!(bSelectionModifier || (GetWorld()->GetRealTimeSeconds() - LastBoxSelectTime) < TouchTapMaxAllowedTime))
+		{
+			// project the touch location and cache the selection point
+			CachedInteraction = CachedSelection = ProjectTouchPointToWorldSpace();
 
-	// project the touch location and cache the selection point
-	CachedSelection = ProjectTouchPointToWorldSpace();
+			// do a selection action with the cached location
+			DoSelectionCommand();
+		}
+
+	}
 
-	// do a selection action with the cached location
-	DoSelectionCommand();
 }
 
 void AStrategyPlayerController::TouchSecondaryStarted(const FInputActionValue& Value)
@@ -311,9 +363,23 @@ void AStrategyPlayerController::TouchSecondaryStarted(const FInputActionValue& V
 
 void AStrategyPlayerController::TouchSecondaryTriggered(const FInputActionValue& Value)
 {
-
 	// update the current position for the second finger
 	CurrentSecondFingerPosition = Value.Get<FVector2D>();
+
+	// are we box selecting, and the finger has moved enough on the touchscreen?
+	if (bSelectionModifier && !StartingSecondFingerPosition.Equals(CurrentSecondFingerPosition, 10.0f))
+	{
+		// update the current interaction position
+		CurrentInteractionPosition = CurrentSecondFingerPosition;
+
+		// update the selection box on the HUD
+		if (StrategyHUD)
+		{
+			const FVector2D DragSize = CurrentSecondFingerPosition - StartingSecondFingerPosition;
+
+			StrategyHUD->DragSelectUpdate(StartingInteractionPosition, DragSize, CurrentSecondFingerPosition, true);
+		}
+	}
 }
 
 void AStrategyPlayerController::TouchSecondaryCompleted(const FInputActionValue& Value)
@@ -321,24 +387,14 @@ void AStrategyPlayerController::TouchSecondaryCompleted(const FInputActionValue&
 
 	// lower the selection modifier flag
 	bSelectionModifier = false;
-}
 
-void AStrategyPlayerController::TouchDoubleTap(const FInputActionValue& Value)
-{
+	// save the last box selection time
+	LastBoxSelectTime = GetWorld()->GetRealTimeSeconds();
 
-	// raise the double tap flag
-	bDoubleTapActive = true;
-
-	// is the selection modifier mode active?
-	if (bSelectionModifier)
+	// hide the selection box on the HUD
+	if (StrategyHUD)
 	{
-		// deselect all units
-		DoDeselectAllCommand();
-
-	} else {
-
-		// select all units on screen
-		DoSelectAllOnScreenCommand();
+		StrategyHUD->DragSelectUpdate(FVector2D::ZeroVector, FVector2D::ZeroVector, FVector2D::ZeroVector, false);
 	}
 }
 
@@ -392,7 +448,8 @@ void AStrategyPlayerController::DoSelectionCommand()
 				TargetUnit->UnitDeselected();
 
 			}
-			else {
+			else
+			{
 
 				// add the unit to the controlled list
 				ControlledUnits.Add(TargetUnit);
@@ -403,24 +460,15 @@ void AStrategyPlayerController::DoSelectionCommand()
 			}
 		}
 
-	} else {
+	}
+	else
+	{
 
 		// are we using touch input?
 		if (InputMode == SIM_Touch)
 		{
-			// is the double tap select all flag set?
-			if (bDoubleTapActive)
-			{
-
-				// release double tap select all
-				bDoubleTapActive = false;
-
-			} else {
-
-				// move all selected units to the target location
-				DoMoveUnitsCommand();
-
-			}
+			// move all selected units to the target location
+			DoMoveUnitsCommand();
 		}
 
 	}
@@ -488,7 +536,9 @@ void AStrategyPlayerController::DoDragScrollCommand()
 		// read the mouse position
 		bool bResult = GetMousePosition(WorkingPosition.X, WorkingPosition.Y);
 
-	} else {
+	}
+	else
+	{
 
 		// read the touch 1 position
 		bool bPressed;
@@ -521,7 +571,8 @@ void AStrategyPlayerController::DoMoveUnitsCommand()
 		CurrentMoveGoal = CachedInteraction;
 
 	}
-	else {
+	else
+	{
 
 		// set the cached selection as our move goal
 		CurrentMoveGoal = CachedSelection;
@@ -648,7 +699,9 @@ AStrategyUnit* AStrategyPlayerController::GetClosestSelectedUnitToLocation(FVect
 					Closest = Dist;
 				}
 
-			} else {
+			}
+			else
+			{
 
 				// no previously selected unit, so use this one
 				OutUnit = CurrentUnit;
@@ -723,3 +776,42 @@ void AStrategyPlayerController::ResetInteraction()
 {
 	bAllowInteraction = true;
 }
+
+void AStrategyPlayerController::CheckTouchTap(bool& bTapped, bool& bDoubleTapped)
+{
+	// get the current game time
+	const float GameTime = GetWorld()->GetRealTimeSeconds();
+
+	// if the player released touch before the max allowed time since press, we have a tap
+	bTapped = (GameTime - LastTapPressTime) < TouchTapMaxAllowedTime;
+
+	if (bTapped)
+	{
+		// we have a double tap if another tap happened before the last release time
+		if ((GameTime - LastTapReleaseTime) < TouchDoubleTapMaxAllowedTime)
+		{
+			// increase the tap counter
+			++TapCount;
+
+		}
+		else
+		{
+
+			// reset the tap counter
+			TapCount = 0;
+		}
+
+	}
+	else
+	{
+
+		// reset the tap counter
+		TapCount = 0;
+	}
+
+	// we have a double tap if the tap count is not zero
+	bDoubleTapped = TapCount >= 1;
+
+	// save the tap release time
+	LastTapReleaseTime = GameTime;
+}
diff --git a/Templates/TP_TopDown/Source/TP_TopDown/Variant_Strategy/StrategyPlayerController.h b/Templates/TP_TopDown/Source/TP_TopDown/Variant_Strategy/StrategyPlayerController.h
--- a/Templates/TP_TopDown/Source/TP_TopDown/Variant_Strategy/StrategyPlayerController.h
+++ b/Templates/TP_TopDown/Source/TP_TopDown/Variant_Strategy/StrategyPlayerController.h
@@ -93,10 +93,6 @@ protected:
 	UPROPERTY(EditAnywhere, Category="Input")
 	UInputAction* SelectionModifierAction;
 
-	/** Input Action for primary touch tap */
-	UPROPERTY(EditAnywhere, Category="Input")
-	UInputAction* TouchPrimaryTapAction;
-
 	/** Input Action for primary touch hold */
 	UPROPERTY(EditAnywhere, Category="Input")
 	UInputAction* TouchPrimaryHoldAction;
@@ -105,10 +101,6 @@ protected:
 	UPROPERTY(EditAnywhere, Category="Input")
 	UInputAction* TouchSecondaryAction;
 
-	/** Input Action for touch double tap */
-	UPROPERTY(EditAnywhere, Category="Input")
-	UInputAction* TouchDoubleTapAction;
-
 	/** Max distance to look for nearby units when doing a click or touch interaction */
 	UPROPERTY(EditAnywhere, Category="Input", meta = (ClampMin = 0, ClampMax = 10000, Units = "cm"))
 	float InteractionRadius = 250.0f;
@@ -170,6 +162,29 @@ protected:
 	/** Currently selected unit list */
 	TArray<AStrategyUnit*> ControlledUnits;
 
+	///////////////////////////////////
+	// Touchscreen enhanced input workaround
+
+	/** Game time when the player last started tapping the touchscreen */
+	float LastTapPressTime = 0.0f;
+
+	/** Game time when the player last ended tapping the touchscreen */
+	float LastTapReleaseTime = 0.0f;
+
+	/** Number of successive times the player has tapped the touchscreen */
+	int32 TapCount = 0;
+
+	/** Game time when the player last completed a box select on the touchscreen */
+	float LastBoxSelectTime = 0.0f;
+
+	/** Max time between touch press and release to be considered a tap */
+	UPROPERTY(EditAnywhere, Category = "Touch Input", meta = (ClampMin = 0, ClampMax = 1, Units="s"))
+	float TouchTapMaxAllowedTime = 0.15f;
+
+	/** Max time between touchscreen taps to be considered a double tap */
+	UPROPERTY(EditAnywhere, Category = "Touch Input", meta = (ClampMin = 0, ClampMax = 1, Units="s"))
+	float TouchDoubleTapMaxAllowedTime = 0.4f;
+
 public:
 
 	/** Constructor */
@@ -233,8 +248,8 @@ protected:
 	/** Touch primary finger hold triggered */
 	void TouchPrimaryHoldTriggered(const FInputActionValue& Value);
 
-	/** Touch primary finger tap completed */
-	void TouchPrimaryTap(const FInputActionValue& Value);
+	/** Touch primary finger hold completed */
+	void TouchPrimaryHoldCompleted(const FInputActionValue& Value);
 
 	/** Touch secondary finger started */
 	void TouchSecondaryStarted(const FInputActionValue& Value);
@@ -245,9 +260,6 @@ protected:
 	/** Touch secondary finger completed */
 	void TouchSecondaryCompleted(const FInputActionValue& Value);
 
-	/** Touch primary finger double tap triggered */
-	void TouchDoubleTap(const FInputActionValue& Value);
-
 	/** Attempt to select or deselect units at the cached location */
 	void DoSelectionCommand();
 
@@ -285,4 +297,7 @@ protected:
 
 	/** Resets the interaction flag */
 	void ResetInteraction();
+
+	/** Detects taps and double taps for mobile platforms. */
+	void CheckTouchTap(bool& bTapped, bool& bDoubleTapped);
 };
