Templates - TopDown Strategy Variant
- Implemented new functionality and fixes on CPP version of the template.
- Small updates, fixes and cleanup to BP version of the template.
#rnx

#virtualized
#Synced-CL 52257170

diff --git a/Templates/TP_TopDown/Source/TP_TopDown/Variant_Strategy/EnvQueryContext_MoveGoal.cpp b/Templates/TP_TopDown/Source/TP_TopDown/Variant_Strategy/EnvQueryContext_MoveGoal.cpp
new file mode 100644
--- /dev/null
+++ b/Templates/TP_TopDown/Source/TP_TopDown/Variant_Strategy/EnvQueryContext_MoveGoal.cpp
@@ -0,0 +1,16 @@
+// Copyright Epic Games, Inc. All Rights Reserved.
+
+
+#include "Variant_Strategy/EnvQueryContext_MoveGoal.h"
+#include "Variant_Strategy/StrategyUnit.h"
+#include "EnvironmentQuery/Items/EnvQueryItemType_Point.h"
+
+void UEnvQueryContext_MoveGoal::ProvideContext(FEnvQueryInstance& QueryInstance, FEnvQueryContextData& ContextData) const
+{
+	// get the querying unit
+	if (AStrategyUnit* QuerierActor = Cast<AStrategyUnit>(QueryInstance.Owner.Get()))
+	{
+		// add the last recorded danger location to the context
+		UEnvQueryItemType_Point::SetContextHelper(ContextData, QuerierActor->GetMovementGoal());
+	}
+}
diff --git a/Templates/TP_TopDown/Source/TP_TopDown/Variant_Strategy/EnvQueryContext_MoveGoal.h b/Templates/TP_TopDown/Source/TP_TopDown/Variant_Strategy/EnvQueryContext_MoveGoal.h
new file mode 100644
--- /dev/null
+++ b/Templates/TP_TopDown/Source/TP_TopDown/Variant_Strategy/EnvQueryContext_MoveGoal.h
@@ -0,0 +1,19 @@
+// Copyright Epic Games, Inc. All Rights Reserved.
+
+#pragma once
+
+#include "CoreMinimal.h"
+#include "EnvironmentQuery/EnvQueryContext.h"
+#include "EnvQueryContext_MoveGoal.generated.h"
+
+/**
+ *  Simple EnvQueryContext that returns a Unit's current movement goal location
+ */
+UCLASS()
+class TP_TOPDOWN_API UEnvQueryContext_MoveGoal : public UEnvQueryContext
+{
+	GENERATED_BODY()
+	
+	/** Provides context data to the query instance */
+	virtual void ProvideContext(FEnvQueryInstance& QueryInstance, FEnvQueryContextData& ContextData) const override;
+};
diff --git a/Templates/TP_TopDown/Source/TP_TopDown/Variant_Strategy/StrategyPlayerController.cpp b/Templates/TP_TopDown/Source/TP_TopDown/Variant_Strategy/StrategyPlayerController.cpp
--- a/Templates/TP_TopDown/Source/TP_TopDown/Variant_Strategy/StrategyPlayerController.cpp
+++ b/Templates/TP_TopDown/Source/TP_TopDown/Variant_Strategy/StrategyPlayerController.cpp
@@ -16,6 +16,10 @@
 #include "StrategyUnit.h"
 #include "NavigationSystem.h"
 #include "Engine/OverlapResult.h"
+#include "InputAction.h"
+#include "StrategyTouchControls.h"
+#include "Widgets/Input/SVirtualJoystick.h"
+#include "TP_TopDown.h"
 
 AStrategyPlayerController::AStrategyPlayerController()
 {
@@ -23,6 +27,34 @@ AStrategyPlayerController::AStrategyPlayerController()
 	bShowMouseCursor = true;
 }
 
+void AStrategyPlayerController::BeginPlay()
+{
+	Super::BeginPlay();
+
+	// only spawn touch controls on local player controllers
+	if (IsLocalPlayerController() && ShouldUseTouchControls())
+	{
+		// spawn the mobile controls widget
+		MobileControlsWidget = CreateWidget<UStrategyTouchControls>(this, MobileControlsWidgetClass);
+
+		if (MobileControlsWidget)
+		{
+			// add the controls to the player screen
+			MobileControlsWidget->AddToPlayerScreen(0);
+
+			// set the PC pointer on the mobile controls widget
+			MobileControlsWidget->SetPlayerController(this);
+
+		} else {
+
+			UE_LOG(LogTP_TopDown, Error, TEXT("Could not spawn mobile controls widget."));
+
+		}
+
+	}
+
+}
+
 void AStrategyPlayerController::SetupInputComponent()
 {
 	Super::SetupInputComponent();
@@ -36,14 +68,13 @@ void AStrategyPlayerController::SetupInputComponent()
 			// choose the context based on the input mode
 			UInputMappingContext* ChosenContext = nullptr;
 
-			switch (InputMode)
+			if (ShouldUseTouchControls())
 			{
-			case SIM_Mouse:
-				ChosenContext = MouseMappingContext;
-				break;
-			case SIM_Touch:
 				ChosenContext = TouchMappingContext;
-				break;
+			}
+			else
+			{
+				ChosenContext = MouseMappingContext;
 			}
 
 			Subsystem->AddMappingContext(ChosenContext, 0);
@@ -65,22 +96,20 @@ void AStrategyPlayerController::SetupInputComponent()
 
 			EnhancedInputComponent->BindAction(SelectClickAction, ETriggerEvent::Completed, this, &AStrategyPlayerController::SelectClick);
 
-			EnhancedInputComponent->BindAction(SelectionModifierAction, ETriggerEvent::Triggered, this, &AStrategyPlayerController::SelectionModifier);
-			EnhancedInputComponent->BindAction(SelectionModifierAction, ETriggerEvent::Completed, this, &AStrategyPlayerController::SelectionModifier);
-			EnhancedInputComponent->BindAction(SelectionModifierAction, ETriggerEvent::Canceled, this, &AStrategyPlayerController::SelectionModifier);
+			EnhancedInputComponent->BindAction(SelectClickAdditiveAction, ETriggerEvent::Completed, this, &AStrategyPlayerController::SelectClickAdditive);
+
+			EnhancedInputComponent->BindAction(SelectAllDoubleClickAction, ETriggerEvent::Completed, this, &AStrategyPlayerController::SelectAllDoubleClick);
 
 			EnhancedInputComponent->BindAction(InteractHoldAction, ETriggerEvent::Started, this, &AStrategyPlayerController::InteractHoldStarted);
 			EnhancedInputComponent->BindAction(InteractHoldAction, ETriggerEvent::Triggered, this, &AStrategyPlayerController::InteractHoldTriggered);
 
-			EnhancedInputComponent->BindAction(InteractClickAction, ETriggerEvent::Started, this, &AStrategyPlayerController::InteractClickStarted);
-			EnhancedInputComponent->BindAction(InteractClickAction, ETriggerEvent::Completed, this, &AStrategyPlayerController::InteractClickCompleted);
+			EnhancedInputComponent->BindAction(InteractClickAction, ETriggerEvent::Completed, this, &AStrategyPlayerController::InteractClick);
 
 			// Touch Interaction
 			EnhancedInputComponent->BindAction(TouchPrimaryHoldAction, ETriggerEvent::Started, this, &AStrategyPlayerController::TouchPrimaryHoldStarted);
 			EnhancedInputComponent->BindAction(TouchPrimaryHoldAction, ETriggerEvent::Triggered, this, &AStrategyPlayerController::TouchPrimaryHoldTriggered);
 			EnhancedInputComponent->BindAction(TouchPrimaryHoldAction, ETriggerEvent::Completed, this, &AStrategyPlayerController::TouchPrimaryHoldCompleted);
 
-			EnhancedInputComponent->BindAction(TouchSecondaryAction, ETriggerEvent::Started, this, &AStrategyPlayerController::TouchSecondaryStarted);
 			EnhancedInputComponent->BindAction(TouchSecondaryAction, ETriggerEvent::Triggered, this, &AStrategyPlayerController::TouchSecondaryTriggered);
 			EnhancedInputComponent->BindAction(TouchSecondaryAction, ETriggerEvent::Completed, this, &AStrategyPlayerController::TouchSecondaryCompleted);
 			EnhancedInputComponent->BindAction(TouchSecondaryAction, ETriggerEvent::Canceled, this, &AStrategyPlayerController::TouchSecondaryCompleted);
@@ -94,15 +123,21 @@ void AStrategyPlayerController::OnPossess(APawn* InPawn)
 	Super::OnPossess(InPawn);
 
 	// ensure we have the right pawn type
-	ControlledPawn = Cast<AStrategyPawn>(InPawn);
-	check(ControlledPawn);
+	ControlledCameraPawn = Cast<AStrategyPawn>(InPawn);
+	check(ControlledCameraPawn);
 
 	// set the zoom level from the pawn's camera
-	DefaultZoom = CameraZoom = ControlledPawn->GetCamera()->OrthoWidth;
+	DefaultZoom = CameraZoom = ControlledCameraPawn->GetCamera()->OrthoWidth;
 
 	// cast the HUD pointer
 	StrategyHUD = Cast<AStrategyHUD>(GetHUD());
 	check(StrategyHUD);
+
+	// if we have a touch controls widget, sync the camera zoom
+	if (MobileControlsWidget)
+	{
+		MobileControlsWidget->BP_SetZoomPercentage(GetDefaultZoomPercentage());
+	}
 }
 
 void AStrategyPlayerController::DragSelectUnits(const TArray<AStrategyUnit*>& Units)
@@ -111,7 +146,7 @@ void AStrategyPlayerController::DragSelectUnits(const TArray<AStrategyUnit*>& Un
 	if (Units.Num() > 0)
 	{
 		// ensure any previous units are deselected
-		DoDeselectAllCommand();
+		DoDeselectAllUnitsCommand();
 
 		// select each new unit
 		for (AStrategyUnit* CurrentUnit : Units)
@@ -130,7 +165,7 @@ void AStrategyPlayerController::DragSelectUnits(const TArray<AStrategyUnit*>& Un
 		// release any currently selected units since nothing is on the box
 		if (ControlledUnits.Num() > 0)
 		{
-			DoDeselectAllCommand();
+			DoDeselectAllUnitsCommand();
 		}
 
 	}
@@ -141,6 +176,18 @@ const TArray<AStrategyUnit*>& AStrategyPlayerController::GetSelectedUnits()
 	return ControlledUnits;
 }
 
+float AStrategyPlayerController::GetDefaultZoomPercentage() const
+{
+	float ZoomPct = (DefaultZoom - MinZoomLevel) / (MaxZoomLevel - MinZoomLevel);
+	return FMath::Clamp(ZoomPct, 0.0f, 1.0f);
+}
+
+bool AStrategyPlayerController::ShouldUseTouchControls() const
+{
+	// are we on a mobile platform? Should we force touch?
+	return SVirtualJoystick::ShouldDisplayTouchInterface() || bForceTouchControls;
+}
+
 void AStrategyPlayerController::MoveCamera(const FInputActionValue& Value)
 {
 	FVector2D InputVector = Value.Get<FVector2D>();
@@ -151,62 +198,49 @@ void AStrategyPlayerController::MoveCamera(const FInputActionValue& Value)
 
 	// get the right input component vector
 	FRotator RightRot = GetControlRotation();
-	ForwardRot.Pitch = 0.0f;
-	ForwardRot.Roll = 0.0f;
+	RightRot.Pitch = 0.0f;
+	RightRot.Roll = 0.0f;
 
 	// add the forward input
-	ControlledPawn->AddMovementInput(ForwardRot.RotateVector(FVector::ForwardVector), InputVector.X + InputVector.Y);
-
-	// add the right input
-	ControlledPawn->AddMovementInput(RightRot.RotateVector(FVector::RightVector), InputVector.X - InputVector.Y);
+	if (ControlledCameraPawn)
+	{
+		ControlledCameraPawn->AddMovementInput(ForwardRot.RotateVector(FVector::ForwardVector), InputVector.X + InputVector.Y);
 
+		// add the right input
+		ControlledCameraPawn->AddMovementInput(RightRot.RotateVector(FVector::RightVector), InputVector.X - InputVector.Y);
+	}
 }
 
 void AStrategyPlayerController::ZoomCamera(const FInputActionValue& Value)
 {
-	// scale the input and subtract from the current zoom level
-	float ZoomLevel = CameraZoom - (Value.Get<float>() * ZoomScaling);
-
-	// clamp to min/max zoom levels
-	CameraZoom = FMath::Clamp(ZoomLevel, MinZoomLevel, MaxZoomLevel);
-
-	// update the pawn's camera
-	ControlledPawn->SetZoomModifier(CameraZoom);
-
+	DoCameraModifyZoomCommand(Value.Get<float>() * ZoomScaling);
 }
 
 void AStrategyPlayerController::ResetCamera(const FInputActionValue& Value)
 {
-	// reset zoom level to its initial value
-	CameraZoom = DefaultZoom;
-
-	// update the pawn's camera
-	ControlledPawn->SetZoomModifier(DefaultZoom);
-
+	DoCameraResetZoomCommand();
 }
 
 void AStrategyPlayerController::SelectHoldStarted(const FInputActionValue& Value)
 {
-	// save the selection start position
-	StartingSelectionPosition = GetMouseLocation();
+	// save the box selection start position
+	StartingBoxSelectionPosition = GetMouseLocationForPlayer();
 
 }
 
 void AStrategyPlayerController::SelectHoldTriggered(const FInputActionValue& Value)
 {
-
 	// get the current mouse position
-	FVector2D SelectionPosition = GetMouseLocation();
+	FVector2D SelectionPosition = GetMouseLocationForPlayer();
 
 	// calculate the size of the selection box
-	FVector2D SelectionSize = SelectionPosition - StartingSelectionPosition;
+	FVector2D SelectionSize = SelectionPosition - StartingBoxSelectionPosition;
 
 	// update the selection box on the HUD
 	if (StrategyHUD)
 	{
-		StrategyHUD->DragSelectUpdate(StartingSelectionPosition, SelectionSize, SelectionPosition, true);
-	}
-	
+		StrategyHUD->DragSelectUpdate(StartingBoxSelectionPosition, SelectionSize, SelectionPosition, true);
+	}	
 }
 
 void AStrategyPlayerController::SelectHoldCompleted(const FInputActionValue& Value)
@@ -220,457 +254,315 @@ void AStrategyPlayerController::SelectHoldCompleted(const FInputActionValue& Val
 
 void AStrategyPlayerController::SelectClick(const FInputActionValue& Value)
 {
+	// get the cursor location
+	FVector CursorLocation;
 
-	if (GetLocationUnderCursor(CachedSelection))
+	if (GetLocationUnderCursor(CursorLocation))
 	{
-		DoSelectionCommand();
+		// select at the cursor
+		DoSelectCommand(CursorLocation, false);
 	}
 }
 
-void AStrategyPlayerController::SelectionModifier(const FInputActionValue& Value)
+void AStrategyPlayerController::SelectClickAdditive(const FInputActionValue& Value)
 {
+	// get the cursor location
+	FVector CursorLocation;
+
+	if (GetLocationUnderCursor(CursorLocation))
+	{
+		// additive select at the cursor
+		DoSelectCommand(CursorLocation, true);
+	}
+}
 
-	// update the selection modifier flag
-	bSelectionModifier = Value.Get<bool>();
+void AStrategyPlayerController::SelectAllDoubleClick(const FInputActionValue& Value)
+{
+	DoSelectAllUnitsOnScreenCommand();
 }
 
 void AStrategyPlayerController::InteractHoldStarted(const FInputActionValue& Value)
 {
 
 	// save the starting interaction position
-	StartingInteractionPosition = GetMouseLocation();
+	StartingDragScrollPosition = GetMouseLocationForPlayer();
 }
 
 void AStrategyPlayerController::InteractHoldTriggered(const FInputActionValue& Value)
 {
-
 	// do a drag scroll 
-	DoDragScrollCommand();
-}
-
-void AStrategyPlayerController::InteractClickStarted(const FInputActionValue& Value)
-{
-
-	// reset the interaction flag
-	ResetInteraction();
+	DoCameraDragScrollCommand(GetMouseLocationForPlayer());
 }
 
-void AStrategyPlayerController::InteractClickCompleted(const FInputActionValue& Value)
+void AStrategyPlayerController::InteractClick(const FInputActionValue& Value)
 {
+	// get the cursor location
+	FVector CursorLocation;
 
-	// do we have any units in the control list and a valid interaction location under the cursor?
-	if (ControlledUnits.Num() > 0 && GetLocationUnderCursor(CachedInteraction))
+	// do we have a valid interaction location under the cursor?
+	if (GetLocationUnderCursor(CursorLocation))
 	{
-		// is double tap select all active?
-		if (bDoubleTapActive)
-		{
-			// release double tap select all
-			bDoubleTapActive = false;
-
-		}
-		else
-		{
-
-			// move the selected units to the target location
-			DoMoveUnitsCommand();
-
-		}
+		// move the selected units to the target location
+		DoMoveUnitsCommand(CursorLocation);
 	}
 }
 
 void AStrategyPlayerController::TouchPrimaryHoldStarted(const FInputActionValue& Value)
 {
-	// save the tap press time
-	LastTapPressTime = GetWorld()->GetRealTimeSeconds();
+	// save the camera drag screen coords
+	StartingDragScrollPosition = Value.Get<FVector2D>();
 
-	// save the starting interaction position
-	StartingInteractionPosition = Value.Get<FVector2D>();
+	
 }
 
-void AStrategyPlayerController::TouchPrimaryHoldTriggered(const FInputActionValue& Value)
+void AStrategyPlayerController::TouchPrimaryHoldTriggered(const FInputActionInstance& Instance)
 {
-	// is this touch longer than a tap?
-	if ((GetWorld()->GetRealTimeSeconds() - LastTapPressTime) > TouchTapMaxAllowedTime)
+	FVector2D InputVector = Instance.GetValue().Get<FVector2D>();
+
+	// update the box select start position
+	StartingBoxSelectionPosition = InputVector;
+
+	if (Instance.GetElapsedTime() > TouchDragScrollHoldTime)
 	{
-		// if we're not doing a box select, do a drag scroll
-		if (!bSelectionModifier)
-		{
-			DoDragScrollCommand();
-		}
+		DoCameraDragScrollCommand(InputVector);
+
+		// save the game time
+		LastTouchDragScrollTime = GetWorld()->GetTimeSeconds();
 	}
 }
 
 void AStrategyPlayerController::TouchPrimaryHoldCompleted(const FInputActionValue& Value)
 {
-	// check if we're doing a tap or double tap.
-	// we have to do this manually because EnhancedInput tap triggers work differently on touch inputs
-	bool bTapped = false;
-	bool bDoubleTapped = false;
-
-	CheckTouchTap(bTapped, bDoubleTapped);
-
-	// do we have a double tap?
-	if (bTapped)
-	{
-		if (bDoubleTapped)
-		{
-			// ensure are we not doing a box select
-			if (!bSelectionModifier)
-			{
-				// depending on the double tap toggle, select or deselect all units
-				if (bDoubleTapActive)
-				{
-					DoDeselectAllCommand();
-				}
-				else
-				{
-					DoSelectAllOnScreenCommand();
-				}
-
-				// toggle the double tap flag
-				bDoubleTapActive = !bDoubleTapActive;
-			}
-		}
-
-	// no double tap, handle this touch input normally
-	}
-	else
+	// ensure we don't trigger a tap input right after we finish a drag scroll
+	if (GetWorld()->GetTimeSeconds() - LastTouchDragScrollTime > 0.1f)
 	{
+		// get the touch location in world space
+		FVector TouchLocation = ProjectTouchPointToWorldSpace();
 
-		// ensure we're not already box selecting, or were just box selecting
-		if (!(bSelectionModifier || (GetWorld()->GetRealTimeSeconds() - LastBoxSelectTime) < TouchTapMaxAllowedTime))
+		// try to do a select command
+		if (!DoSelectCommand(TouchLocation, true))
 		{
-			// project the touch location and cache the selection point
-			CachedInteraction = CachedSelection = ProjectTouchPointToWorldSpace();
-
-			// do a selection action with the cached location
-			DoSelectionCommand();
+			// if nothing was selected, do a move units command instead
+			DoMoveUnitsCommand(TouchLocation);
 		}
-
 	}
-
-}
-
-void AStrategyPlayerController::TouchSecondaryStarted(const FInputActionValue& Value)
-{
-
-	// raise the selection modifier flag
-	bSelectionModifier = true;
-
-	// save the starting position for the second finger
-	StartingSecondFingerPosition = Value.Get<FVector2D>();
 }
 
 void AStrategyPlayerController::TouchSecondaryTriggered(const FInputActionValue& Value)
 {
-	// update the current position for the second finger
-	CurrentSecondFingerPosition = Value.Get<FVector2D>();
+	// get the touch 2 screen coords
+	FVector2D SelectionPosition = Value.Get<FVector2D>();
 
-	// are we box selecting, and the finger has moved enough on the touchscreen?
-	if (bSelectionModifier && !StartingSecondFingerPosition.Equals(CurrentSecondFingerPosition, 10.0f))
-	{
-		// update the current interaction position
-		CurrentInteractionPosition = CurrentSecondFingerPosition;
-
-		// update the selection box on the HUD
-		if (StrategyHUD)
-		{
-			const FVector2D DragSize = CurrentSecondFingerPosition - StartingSecondFingerPosition;
+	// calculate the size of the selection box
+	FVector2D SelectionSize = SelectionPosition - StartingBoxSelectionPosition;
 
-			StrategyHUD->DragSelectUpdate(StartingInteractionPosition, DragSize, CurrentSecondFingerPosition, true);
-		}
+	// update the selection box on the HUD
+	if (StrategyHUD)
+	{
+		StrategyHUD->DragSelectUpdate(StartingBoxSelectionPosition, SelectionSize, SelectionPosition, true);
 	}
+		
 }
 
 void AStrategyPlayerController::TouchSecondaryCompleted(const FInputActionValue& Value)
 {
-
-	// lower the selection modifier flag
-	bSelectionModifier = false;
-
-	// save the last box selection time
-	LastBoxSelectTime = GetWorld()->GetRealTimeSeconds();
-
-	// hide the selection box on the HUD
 	if (StrategyHUD)
 	{
+		// hide the selection box
 		StrategyHUD->DragSelectUpdate(FVector2D::ZeroVector, FVector2D::ZeroVector, FVector2D::ZeroVector, false);
 	}
 }
 
-void AStrategyPlayerController::DoSelectionCommand()
+bool AStrategyPlayerController::DoSelectCommand(const FVector& SelectLocation, bool bAdditiveSelection)
 {
+	// deselect any units unless this is an additive selection
+	if (!bAdditiveSelection)
+	{
+		DoDeselectAllUnitsCommand();
+	}
 
-	// do a sphere sweep to look for actors to select
-	FHitResult OutHit;
-
-	const FVector Start = CachedSelection;
-	const FVector End = Start + FVector::UpVector * 350.0f;
+	// do an overlap test at the cursor location
+	TArray<FOverlapResult> OutOverlaps;
 
-	FCollisionShape InteractionSphere;
-	InteractionSphere.SetSphere(InteractionRadius);
+	FCollisionShape CollisionSphere;
+	CollisionSphere.SetSphere(SelectionRadius);
 
 	FCollisionObjectQueryParams ObjectParams;
 	ObjectParams.AddObjectTypesToQuery(ECC_Pawn);
 
 	FCollisionQueryParams QueryParams;
-	QueryParams.AddIgnoredActor(this);
-	QueryParams.AddIgnoredActor(GetPawn());
-	QueryParams.bTraceComplex = true;
 
-	GetWorld()->SweepSingleByObjectType(OutHit, Start, End, FQuat::Identity, ObjectParams, InteractionSphere, QueryParams);
-
-	// if we're using the mouse and are not holding the selection modifier key, deselect any units first
-	if (InputMode == SIM_Mouse && !bSelectionModifier)
+	if (GetWorld()->OverlapMultiByObjectType(OutOverlaps, SelectLocation, FQuat::Identity, ObjectParams, CollisionSphere, QueryParams))
 	{
-
-		DoDeselectAllCommand();
-	}
-
-	// did we hit a unit?
-	if (OutHit.bBlockingHit)
-	{
-
-		// update the target unit
-		TargetUnit = Cast<AStrategyUnit>(OutHit.GetActor());
-
-		if (TargetUnit)
+		// find the first unit we've overlapped
+		for (const FOverlapResult& CurrentOverlap : OutOverlaps)
 		{
-
-			// is the unit already in the controlled list?
-			if (ControlledUnits.Contains(TargetUnit))
-			{
-
-				// remove the units from the controlled list
-				ControlledUnits.Remove(TargetUnit);
-
-				// tell the unit it's been deselected
-				TargetUnit->UnitDeselected();
-
-			}
-			else
+			if (AStrategyUnit* CurrentUnit = Cast<AStrategyUnit>(CurrentOverlap.GetActor()))
 			{
+				// is this unit already selected?
+				if (ControlledUnits.Contains(CurrentUnit))
+				{
+					// deselect the unit
+					ControlledUnits.Remove(CurrentUnit);
 
-				// add the unit to the controlled list
-				ControlledUnits.Add(TargetUnit);
+					CurrentUnit->UnitDeselected();
+				}
+				else
+				{
+					// select the unit
+					ControlledUnits.Add(CurrentUnit);
 
-				// tell the unit it's been selected
-				TargetUnit->UnitSelected();
+					CurrentUnit->UnitSelected();
+				}
 
+				// found a unit
+				return true;
 			}
 		}
-
 	}
-	else
-	{
-
-		// are we using touch input?
-		if (InputMode == SIM_Touch)
-		{
-			// move all selected units to the target location
-			DoMoveUnitsCommand();
-		}
 
-	}
+	// didn't find a unit
+	return false;
 }
 
-void AStrategyPlayerController::DoSelectAllOnScreenCommand()
+void AStrategyPlayerController::DoSelectAllUnitsOnScreenCommand()
 {
+	// get all units on the level
+	TArray<AActor*> Units;
 
-	// find all NPCs currently on screen
-	TArray<AActor*> FoundActors;
-	UGameplayStatics::GetAllActorsOfClass(GetWorld(), AStrategyUnit::StaticClass(), FoundActors);
+	UGameplayStatics::GetAllActorsOfClass(this, AStrategyUnit::StaticClass(), Units);
 
-	// process each actor found
-	for (AActor* CurrentActor : FoundActors)
+	// process each unit
+	for (AActor* CurrentActor : Units)
 	{
-		// cast back to our unit class
 		if (AStrategyUnit* CurrentUnit = Cast<AStrategyUnit>(CurrentActor))
 		{
-			// has the actor been recently rendered?
-			if (CurrentActor->WasRecentlyRendered(0.2f))
+			// is the unit is not already selected, and is on screen?
+			if (!ControlledUnits.Contains(CurrentUnit) && CurrentUnit->WasRecentlyRendered(0.2f))
 			{
+				// select the unit
+				ControlledUnits.Add(CurrentUnit);
 
-				// is the actor not on our controlled units list?
-				if (!ControlledUnits.Contains(CurrentUnit))
-				{
-					// add it to the controlled units list
-					ControlledUnits.Add(CurrentUnit);
-
-					// notify it of selection
-					CurrentUnit->UnitSelected();
-				}
+				CurrentUnit->UnitSelected();
 			}
-		}		
+		}
+		
 	}
-
 }
 
-void AStrategyPlayerController::DoDeselectAllCommand()
+void AStrategyPlayerController::DoDeselectAllUnitsCommand()
 {
-
-	// tell each controlled unit it's been deselected
+	// deselect each unit
 	for (AStrategyUnit* CurrentUnit : ControlledUnits)
 	{
-		// ensure the unit hasn't been destroyed
 		if (IsValid(CurrentUnit))
 		{
-
 			CurrentUnit->UnitDeselected();
 		}
 	}
 
-	// clear the controlled units list
+	// clear the selection list
 	ControlledUnits.Empty();
 }
 
-void AStrategyPlayerController::DoDragScrollCommand()
+void AStrategyPlayerController::DoToggleSelectAllUnitsCommand()
 {
-
-	// choose the cursor position based on the input mode
-	FVector2D WorkingPosition;
-	
-	if (InputMode == EStrategyInputMode::SIM_Mouse)
+	// do we have units selected?
+	if (ControlledUnits.Num() > 0)
 	{
-
-		// read the mouse position
-		bool bResult = GetMousePosition(WorkingPosition.X, WorkingPosition.Y);
-
+		// deselect all units
+		DoDeselectAllUnitsCommand();
 	}
 	else
 	{
-
-		// read the touch 1 position
-		bool bPressed;
-		GetInputTouchState(ETouchIndex::Touch1, WorkingPosition.X, WorkingPosition.Y, bPressed);
-
+		// select all units on screen
+		DoSelectAllUnitsOnScreenCommand();
 	}
-
-	// find the difference between the starting interaction position and current coords
-	const FVector2D InteractionDelta = StartingInteractionPosition - WorkingPosition;
-
-	const FRotator CameraRot(0.0f, -45.0f, 0.0f);
-
-	// rotate and scale the interaction delta
-	const FVector ScrollDelta = CameraRot.RotateVector(FVector(InteractionDelta.X, InteractionDelta.Y, 0.0f)) * DragMultiplier;
-
-	// apply the world offset to the controlled pawn
-	ControlledPawn->AddActorWorldOffset(ScrollDelta);
 }
 
-void AStrategyPlayerController::DoMoveUnitsCommand()
+void AStrategyPlayerController::DoCameraDragScrollCommand(const FVector2D& CurrentCursorPosition)
 {
+	// subtract the starting position from the cursor to find the on-screen movement delta
+	FVector2D MoveDelta = StartingDragScrollPosition - CurrentCursorPosition;
 
-	// set the movement goal
-	FVector CurrentMoveGoal;
+	// rotate the movement delta to match the isometric perspective
+	const FRotator IsoRotation(0.0f, -45.0f, 0.0f);
 
-	if (InputMode == EStrategyInputMode::SIM_Mouse)
-	{
+	FVector RotatedDelta = IsoRotation.RotateVector(FVector(MoveDelta.X, MoveDelta.Y, 0.0f));
 
-		// set the cached interaction point as our move goal
-		CurrentMoveGoal = CachedInteraction;
+	// apply drag
+	RotatedDelta *= DragMultiplier;
 
-	}
-	else
+	// apply the offset to the camera pawn
+	if (ControlledCameraPawn)
 	{
-
-		// set the cached selection as our move goal
-		CurrentMoveGoal = CachedSelection;
-
+		ControlledCameraPawn->AddActorWorldOffset(RotatedDelta);
 	}
+}
 
-	// get the closest selected unit to the move goal. This will be our lead unit
-	AStrategyUnit* Closest = GetClosestSelectedUnitToLocation(CurrentMoveGoal);
-
-	// this will be set to true if any of the move requests fail
-	bool bInteractionFailed = false;
-
-	// process each unit in the controlled list
-	for (AStrategyUnit* CurrentUnit : ControlledUnits)
+void AStrategyPlayerController::DoMoveUnitsCommand(const FVector& GoalLocation)
+{
+	if (ControlledUnits.Num() > 0)
 	{
-		if (IsValid(CurrentUnit))
-		{
-
-			// stop the unit
-			CurrentUnit->StopMoving();
-
-			// move the lead unit to the goal, all other units to random navigable points around it
-			FVector MoveGoal = CurrentMoveGoal;
-
-			if (CurrentUnit != Closest)
-			{
-
-				UNavigationSystemV1::K2_GetRandomLocationInNavigableRadius(GetWorld(), CurrentMoveGoal, MoveGoal, InteractionRadius * 0.66f);
-			}
+		// find the closest unit to the goal
+		AStrategyUnit* ClosestUnit = GetClosestSelectedUnitToLocation(GoalLocation);
 
-			// subscribe to the unit's move completed delegate
-			CurrentUnit->OnMoveCompleted.AddDynamic(this, &AStrategyPlayerController::OnMoveCompleted);
-
-			// set up movement to the goal location
-			if (!CurrentUnit->MoveToLocation(MoveGoal, InteractionRadius * 0.66f))
+		// tell each unit to move to the location
+		for (AStrategyUnit* CurrentUnit : ControlledUnits)
+		{
+			if (IsValid(CurrentUnit))
 			{
-				// the move request failed, so flag it
-				bInteractionFailed = true;
+				CurrentUnit->MoveToLocation(GoalLocation, CurrentUnit == ClosestUnit, ControlledUnits);
 			}
 		}
 
-	}
-
-	// play the cursor feedback depending on whether our move succeeded or not
-	BP_CursorFeedback(CachedInteraction, !bInteractionFailed);
+		// show positive cursor feedback
+		BP_CursorFeedback(GoalLocation, true);
 
+	}
+	else
+	{
+		// no units selected, so just show negative cursor feedback
+		BP_CursorFeedback(GoalLocation, false);
+	}
 }
 
-void AStrategyPlayerController::OnMoveCompleted(AStrategyUnit* MovedUnit)
+void AStrategyPlayerController::DoCameraModifyZoomCommand(float ZoomDelta)
 {
-	// is the unit valid?
-	if (IsValid(MovedUnit))
-	{
-		// unsubscribe from the delegate
-		MovedUnit->OnMoveCompleted.RemoveDynamic(this, &AStrategyPlayerController::OnMoveCompleted);
-		
-		// skip if interactions are locked
-		if (!bAllowInteraction)
-		{
-			return;
-		}
-
-		// disallow additional interactions until we reset
-		bAllowInteraction = false;
-
-		// is the unit close enough to the cached interaction location?
-		if(FVector::Dist2D(CachedInteraction, MovedUnit->GetActorLocation()) < InteractionRadius)
-		{
+	// add the delta
+	CameraZoom += ZoomDelta;
 
-			// do an overlap test to find nearby interactive objects
-			TArray<FOverlapResult> OutOverlaps;
+	// clamp between min and max
+	CameraZoom = FMath::Clamp(CameraZoom, MinZoomLevel, MaxZoomLevel);
 
-			FCollisionShape CollisionSphere;
-			CollisionSphere.SetSphere(InteractionRadius);
+	// set the zoom on the camera pawn
+	if (ControlledCameraPawn)
+	{
+		ControlledCameraPawn->SetZoomModifier(CameraZoom);
+	}
+}
 
-			FCollisionObjectQueryParams ObjectParams;
-			ObjectParams.AddObjectTypesToQuery(ECC_WorldDynamic);
-			
-			FCollisionQueryParams QueryParams;
+void AStrategyPlayerController::DoCameraResetZoomCommand()
+{
+	// reset to default zoom
+	CameraZoom = DefaultZoom;
 
-			QueryParams.AddIgnoredActor(MovedUnit);
+	// set the zoom on the camera pawn
+	if (ControlledCameraPawn)
+	{
+		ControlledCameraPawn->SetZoomModifier(CameraZoom);
+	}
+}
 
-			for(const AStrategyUnit* CurSelected : ControlledUnits)
-			{
-				QueryParams.AddIgnoredActor(CurSelected);
-			}
+void AStrategyPlayerController::DoCameraSetZoomPercentageCommand(float Percentage)
+{
+	// lerp between min and max zoom
+	CameraZoom = FMath::Lerp(MinZoomLevel, MaxZoomLevel, FMath::Clamp(Percentage, 0.0f, 1.0f));
 
-			if (GetWorld()->OverlapMultiByObjectType(OutOverlaps, CachedInteraction, FQuat::Identity, ObjectParams, CollisionSphere, QueryParams))
-			{
-				for (const FOverlapResult& CurrentOverlap : OutOverlaps)
-				{
-					if (AStrategyUnit* CurrentUnit = Cast<AStrategyUnit>(CurrentOverlap.GetActor()))
-					{
-						CurrentUnit->Interact(MovedUnit);
-					}
-				}
-			}
-		}
+	// set the zoom on the camera pawn
+	if (ControlledCameraPawn)
+	{
+		ControlledCameraPawn->SetZoomModifier(CameraZoom);
 	}
 }
 
@@ -683,7 +575,7 @@ AStrategyUnit* AStrategyPlayerController::GetClosestSelectedUnitToLocation(FVect
 	// process each unit on the list
 	for (AStrategyUnit* CurrentUnit : ControlledUnits)
 	{
-		if (CurrentUnit != nullptr)
+		if (IsValid(CurrentUnit))
 		{
 			// have we selected a unit already?
 			if (OutUnit != nullptr)
@@ -717,7 +609,7 @@ AStrategyUnit* AStrategyPlayerController::GetClosestSelectedUnitToLocation(FVect
 	return OutUnit;
 }
 
-FVector2D AStrategyPlayerController::GetMouseLocation()
+FVector2D AStrategyPlayerController::GetMouseLocationForPlayer()
 {
 	// attempt to get the mouse position from this PC
 	float MouseX, MouseY;
@@ -736,7 +628,24 @@ bool AStrategyPlayerController::GetLocationUnderCursor(FVector& Location)
 	// trace the visibility channel at the cursor location
 	FHitResult OutHit;
 
-	GetHitResultUnderCursorByChannel(SelectionTraceChannel, true, OutHit);
+	GetHitResultUnderCursorByChannel(SelectionTraceChannel, false, OutHit);
+
+	// if there was a blocking hit, return the hit location
+	if (OutHit.bBlockingHit)
+	{
+		Location = OutHit.Location;
+		return true;
+	}
+
+	return OutHit.bBlockingHit;
+}
+
+bool AStrategyPlayerController::GetLocationUnderFinger(FVector& Location)
+{
+	// trace the visibility channel at Touch 1 location
+	FHitResult OutHit;
+
+	GetHitResultUnderFingerByChannel(ETouchIndex::Touch1, SelectionTraceChannel, false, OutHit);
 
 	// if there was a blocking hit, return the hit location
 	if (OutHit.bBlockingHit)
@@ -762,56 +671,23 @@ FVector AStrategyPlayerController::ProjectTouchPointToWorldSpace()
 	// deproject the coords into world space
 	if (DeprojectScreenPositionToWorld(TouchX, TouchY, WorldLocation, WorldDirection))
 	{
-		// intersect with a horizontal plane and return the resulting point
-		const FPlane IntersectPlane(FVector::ZeroVector, FVector::UpVector);
-
-		return FMath::LinePlaneIntersection(WorldLocation, WorldLocation + (WorldDirection * 100000.0f), IntersectPlane);
-	}
-
-	// failed to deproject, return a zero vector
-	return FVector::ZeroVector;
-}
-
-void AStrategyPlayerController::ResetInteraction()
-{
-	bAllowInteraction = true;
-}
+		// run a line trace down the camera
+		FHitResult OutHit;
 
-void AStrategyPlayerController::CheckTouchTap(bool& bTapped, bool& bDoubleTapped)
-{
-	// get the current game time
-	const float GameTime = GetWorld()->GetRealTimeSeconds();
-
-	// if the player released touch before the max allowed time since press, we have a tap
-	bTapped = (GameTime - LastTapPressTime) < TouchTapMaxAllowedTime;
+		GetWorld()->LineTraceSingleByChannel(OutHit, WorldLocation, WorldLocation + WorldDirection * 10000.0f, ECC_Visibility);
 
-	if (bTapped)
-	{
-		// we have a double tap if another tap happened before the last release time
-		if ((GameTime - LastTapReleaseTime) < TouchDoubleTapMaxAllowedTime)
+		// if we hit something, return the impact point
+		if (OutHit.bBlockingHit)
 		{
-			// increase the tap counter
-			++TapCount;
-
+			return OutHit.ImpactPoint;
 		}
-		else
-		{
-
-			// reset the tap counter
-			TapCount = 0;
-		}
-
-	}
-	else
-	{
+		
 
-		// reset the tap counter
-		TapCount = 0;
+		// intersect with a horizontal plane and return the resulting point
+		const FPlane IntersectPlane(FVector::ZeroVector, FVector::UpVector);
+		return FMath::LinePlaneIntersection(WorldLocation, WorldLocation + (WorldDirection * 100000.0f), IntersectPlane);
 	}
 
-	// we have a double tap if the tap count is not zero
-	bDoubleTapped = TapCount >= 1;
-
-	// save the tap release time
-	LastTapReleaseTime = GameTime;
-}
+	// failed to deproject, return a zero vector
+	return FVector::ZeroVector;
+}
\ No newline at end of file
diff --git a/Templates/TP_TopDown/Source/TP_TopDown/Variant_Strategy/StrategyPlayerController.h b/Templates/TP_TopDown/Source/TP_TopDown/Variant_Strategy/StrategyPlayerController.h
--- a/Templates/TP_TopDown/Source/TP_TopDown/Variant_Strategy/StrategyPlayerController.h
+++ b/Templates/TP_TopDown/Source/TP_TopDown/Variant_Strategy/StrategyPlayerController.h
@@ -11,16 +11,10 @@ class UInputMappingContext;
 class UNiagaraSystem;
 struct FInputActionValue;
 class AStrategyHUD;
-class AStrategyNPC;
 class UInputAction;
-
-/** Enum to determine the last used input type */
-UENUM(BlueprintType)
-enum EStrategyInputMode : uint8
-{
-	SIM_Mouse	UMETA(DisplayName = "Mouse"),
-	SIM_Touch	UMETA(DisplayName = "Touch")
-};
+struct FInputActionInstance;
+class AStrategyUnit;
+class UStrategyTouchControls;
 
 /**
  *  Player Controller for a top-down strategy game.
@@ -35,15 +29,11 @@ class AStrategyPlayerController : public APlayerController
 protected:
 
 	/** Strategy Pawn associated with this controller */
-	TObjectPtr<AStrategyPawn> ControlledPawn;
+	TObjectPtr<AStrategyPawn> ControlledCameraPawn;
 
 	/** Strategy HUD associated with this controller */
 	TObjectPtr<AStrategyHUD> StrategyHUD;
 
-	/** Determines the chosen input type */
-	UPROPERTY(EditAnywhere, Category="Input")
-	TEnumAsByte<EStrategyInputMode> InputMode = SIM_Mouse;
-
 	/** Input mapping context to use with mouse input */
 	UPROPERTY(EditAnywhere, Category="Input")
 	UInputMappingContext* MouseMappingContext;
@@ -73,10 +63,18 @@ protected:
 	UPROPERTY(EditAnywhere, Category="Input")
 	UInputAction* ResetCameraAction;
 
-	/** Input Action for select and click */
+	/** Input Action for select click */
 	UPROPERTY(EditAnywhere, Category="Input")
 	UInputAction* SelectClickAction;
 
+	/** Input Action for additive select click */
+	UPROPERTY(EditAnywhere, Category="Input")
+	UInputAction* SelectClickAdditiveAction;
+
+	/** Input Action for select all double click */
+	UPROPERTY(EditAnywhere, Category="Input")
+	UInputAction* SelectAllDoubleClickAction;
+
 	/** Input Action for select press and hold */
 	UPROPERTY(EditAnywhere, Category="Input")
 	UInputAction* SelectHoldAction;
@@ -101,40 +99,39 @@ protected:
 	UPROPERTY(EditAnywhere, Category="Input")
 	UInputAction* TouchSecondaryAction;
 
-	/** Max distance to look for nearby units when doing a click or touch interaction */
-	UPROPERTY(EditAnywhere, Category="Input", meta = (ClampMin = 0, ClampMax = 10000, Units = "cm"))
-	float InteractionRadius = 250.0f;
-
-	/** Max distance between the starting and current position of the second touch finger to be considered a box selection */
-	UPROPERTY(EditAnywhere, Category="Input", meta = (ClampMin = 0, ClampMax = 10000))
-	float MinSecondFingerDistanceForBoxSelect = 10.0f;
+	/** Pointer to the mobile controls widget */
+	TObjectPtr<UStrategyTouchControls> MobileControlsWidget;
 
-	/** Saves the world location of the last initiated interaction */
-	FVector CachedInteraction;
+	/** Touch controls widget class to spawn on mobile platforms */
+	UPROPERTY(EditAnywhere, Category="Input")
+	TSubclassOf<UStrategyTouchControls> MobileControlsWidgetClass;
 
-	/** Saves the world location of the last initiated unit selection */
-	FVector CachedSelection;
+	/** If true, the PC will be initialized with touchscreen controls */
+	UPROPERTY(EditAnywhere, Category="Input")
+	bool bForceTouchControls = false;
 
-	/** Saves the world location where the player started a press and hold interaction */
-	FVector2D StartingInteractionPosition;
+	/** Max distance to look for nearby units when doing a click or touch interaction */
+	UPROPERTY(EditAnywhere, Category="Input", meta = (ClampMin = 0, ClampMax = 10000, Units = "cm"))
+	float SelectionRadius = 250.0f;
 
-	/** Saves the current world location of the player's cursor in press and hold interaction */
-	FVector2D CurrentInteractionPosition;
+	/** Cached starting position for camera drag scrolling */
+	FVector2D StartingDragScrollPosition;
 
-	/** Saves the starting world location of a player's cursor in a press and hold selection box */
-	FVector2D StartingSelectionPosition;
+	/** Cached starting position for box select */
+	FVector2D StartingBoxSelectionPosition;
 
-	/** Saves the starting location of a two-finger touch interaction (pinch) */
-	FVector2D StartingSecondFingerPosition;
+	/** Last game time when a drag scroll was performed, so we can avoid spamming commands on drag scroll end */
+	float LastTouchDragScrollTime = 0.0f;
 
-	/** Saves the current location of a two-finger touch interaction (pinch) */
-	FVector2D CurrentSecondFingerPosition;
+	/** Time the finger needs to be held down to initiate a drag scroll */
+	UPROPERTY(EditAnywhere, Category="Input", meta = (ClampMin = 0, ClampMax = 1, Units = "s"))
+	float TouchDragScrollHoldTime = 0.15f;
 
 	/** Current camera zoom level */
-	float CameraZoom;
+	float CameraZoom = 0.0f;
 
 	/** Default camera zoom level */
-	float DefaultZoom;
+	float DefaultZoom = 1000.0f;
 
 	/** Minimum allowed camera zoom level */
 	UPROPERTY(EditAnywhere, Category = "Camera", meta = (ClampMin = 0, ClampMax = 10000))
@@ -156,40 +153,17 @@ protected:
 	UPROPERTY(EditAnywhere, Category = "Selection")
 	TEnumAsByte<ETraceTypeQuery> SelectionTraceChannel;
 
-	/** Currently selected unit */
-	AStrategyUnit* TargetUnit = nullptr;
-
 	/** Currently selected unit list */
 	TArray<AStrategyUnit*> ControlledUnits;
 
-	///////////////////////////////////
-	// Touchscreen enhanced input workaround
-
-	/** Game time when the player last started tapping the touchscreen */
-	float LastTapPressTime = 0.0f;
-
-	/** Game time when the player last ended tapping the touchscreen */
-	float LastTapReleaseTime = 0.0f;
-
-	/** Number of successive times the player has tapped the touchscreen */
-	int32 TapCount = 0;
-
-	/** Game time when the player last completed a box select on the touchscreen */
-	float LastBoxSelectTime = 0.0f;
-
-	/** Max time between touch press and release to be considered a tap */
-	UPROPERTY(EditAnywhere, Category = "Touch Input", meta = (ClampMin = 0, ClampMax = 1, Units="s"))
-	float TouchTapMaxAllowedTime = 0.15f;
-
-	/** Max time between touchscreen taps to be considered a double tap */
-	UPROPERTY(EditAnywhere, Category = "Touch Input", meta = (ClampMin = 0, ClampMax = 1, Units="s"))
-	float TouchDoubleTapMaxAllowedTime = 0.4f;
-
 public:
 
 	/** Constructor */
 	AStrategyPlayerController();
 
+	/** BeginPlay initialization */
+	virtual void BeginPlay() override;
+
 	/** Initialize input bindings */
 	virtual void SetupInputComponent() override;
 
@@ -204,6 +178,16 @@ public:
 	/** Passes the list of selected units */
 	const TArray<AStrategyUnit*>& GetSelectedUnits();
 
+	/** Returns the default camera zoom percentage value */
+	float GetDefaultZoomPercentage() const;
+
+protected:
+
+	/** Returns true if the PC should run using touchscreen controls */
+	bool ShouldUseTouchControls() const;
+
+	// mouse + keyboard input
+
 protected:
 
 	/** Moves the camera by the given input */
@@ -227,8 +211,11 @@ protected:
 	/** Select click action */
 	void SelectClick(const FInputActionValue& Value);
 
-	/** Presses or releases the selection modifier key */
-	void SelectionModifier(const FInputActionValue& Value);
+	/** Select Click Additive Action */
+	void SelectClickAdditive(const FInputActionValue& Value);
+
+	/** Select All Double Click Action */
+	void SelectAllDoubleClick(const FInputActionValue& Value);
 
 	/** Starts an interaction hold input */
 	void InteractHoldStarted(const FInputActionValue& Value);
@@ -237,67 +224,74 @@ protected:
 	void InteractHoldTriggered(const FInputActionValue& Value);
 
 	/** Interaction click input started */
-	void InteractClickStarted(const FInputActionValue& Value);
-
-	/** Interaction click input completed */
-	void InteractClickCompleted(const FInputActionValue& Value);
+	void InteractClick(const FInputActionValue& Value);
 
 	/** Touch primary finger hold started */
 	void TouchPrimaryHoldStarted(const FInputActionValue& Value);
 
 	/** Touch primary finger hold triggered */
-	void TouchPrimaryHoldTriggered(const FInputActionValue& Value);
+	void TouchPrimaryHoldTriggered(const FInputActionInstance& Instance);
 
 	/** Touch primary finger hold completed */
 	void TouchPrimaryHoldCompleted(const FInputActionValue& Value);
 
-	/** Touch secondary finger started */
-	void TouchSecondaryStarted(const FInputActionValue& Value);
-
 	/** Touch secondary finger triggered */
 	void TouchSecondaryTriggered(const FInputActionValue& Value);
 
 	/** Touch secondary finger completed */
 	void TouchSecondaryCompleted(const FInputActionValue& Value);
 
-	/** Attempt to select or deselect units at the cached location */
-	void DoSelectionCommand();
+	// commands
+
+public:
 
-	/** Select all units currently on screen */
-	void DoSelectAllOnScreenCommand();
+	/** Attempt to select or deselect a unit near the given location. Supports additive selection modifier. */
+	bool DoSelectCommand(const FVector& SelectLocation, bool bAdditiveSelection);
 
-	/** Deselect all controlled units */
-	void DoDeselectAllCommand();
+	/** Attempts to select all units on screen */
+	void DoSelectAllUnitsOnScreenCommand();
 
-	/** Drag scroll the camera */
-	void DoDragScrollCommand();
+	/** Deselects any selected units */
+	void DoDeselectAllUnitsCommand();
 
-	/** Move all selected units */
-	void DoMoveUnitsCommand();
+	/** Toggles between selecting all units on screen and deselecting units */
+	void DoToggleSelectAllUnitsCommand();
 
-	/** Called when a unit move is completed */
-	UFUNCTION()
-	void OnMoveCompleted(AStrategyUnit* MovedUnit);
+	/** Scrolls the camera based on a new screen coordinate */
+	void DoCameraDragScrollCommand(const FVector2D& CurrentCursorPosition);
+
+	/** Attempts to move all selected units to the given location */
+	void DoMoveUnitsCommand(const FVector& GoalLocation);
+
+	/** Applies a zoom change to the camera */
+	void DoCameraModifyZoomCommand(float ZoomDelta);
+
+	/** Resets the camera zoom to default */
+	void DoCameraResetZoomCommand();
+
+	/** Sets the camera zoom to a percentage between min and max zoom */
+	void DoCameraSetZoomPercentageCommand(float Percentage);
+
+protected:
 
 	/** Sorts all controlled units based on their distance to the provided world location */
 	AStrategyUnit* GetClosestSelectedUnitToLocation(FVector TargetLocation);
 
 	/** Calculates and returns the current mouse location */
-	FVector2D GetMouseLocation();
+	FVector2D GetMouseLocationForPlayer();
 
 	/** Attempts to get the world location under the cursor, returns true if successful */
 	bool GetLocationUnderCursor(FVector& Location);
 
+	/** Attempts to get the world location under the Touch 1 finger, returns true if successful */
+	bool GetLocationUnderFinger(FVector& Location);
+
 	/** Projects the current touch location into world space */
 	FVector ProjectTouchPointToWorldSpace();
 
+protected:
+
 	/** Spawns the positive cursor effect */
 	UFUNCTION(BlueprintImplementableEvent, Category="Cursor", meta = (DisplayName="Cursor Feedback"))
 	void BP_CursorFeedback(FVector Location, bool bPositive);
-
-	/** Resets the interaction flag */
-	void ResetInteraction();
-
-	/** Detects taps and double taps for mobile platforms. */
-	void CheckTouchTap(bool& bTapped, bool& bDoubleTapped);
 };
diff --git a/Templates/TP_TopDown/Source/TP_TopDown/Variant_Strategy/StrategyUnit.cpp b/Templates/TP_TopDown/Source/TP_TopDown/Variant_Strategy/StrategyUnit.cpp
--- a/Templates/TP_TopDown/Source/TP_TopDown/Variant_Strategy/StrategyUnit.cpp
+++ b/Templates/TP_TopDown/Source/TP_TopDown/Variant_Strategy/StrategyUnit.cpp
@@ -7,6 +7,9 @@
 #include "Kismet/KismetMathLibrary.h"
 #include "Components/SphereComponent.h"
 #include "Navigation/PathFollowingComponent.h"
+#include "EnvironmentQuery/EnvQueryManager.h"
+#include "EnvironmentQuery/EnvQueryInstanceBlueprintWrapper.h"
+#include "Engine/OverlapResult.h"
 
 AStrategyUnit::AStrategyUnit()
 {
@@ -59,6 +62,9 @@ void AStrategyUnit::StopMoving()
 {
 	// use the character movement component to stop movement
 	GetCharacterMovement()->StopMovementImmediately();
+
+	// stop the unit's interaction animation
+	BP_StopAnimation();
 }
 
 void AStrategyUnit::UnitSelected()
@@ -90,57 +96,121 @@ void AStrategyUnit::Interact(AStrategyUnit* Interactor)
 	
 }
 
-bool AStrategyUnit::MoveToLocation(const FVector& Location, float AcceptanceRadius)
+void AStrategyUnit::MoveToLocation(const FVector& Location, bool bInteract, const TArray<AStrategyUnit*> IgnoreList)
 {
-	// ensure we have a valid AI Controller
-	if (AIController)
-	{
-		// set up the AI Move Request
-		FAIMoveRequest MoveReq;
-
-		MoveReq.SetGoalLocation(Location);
-		MoveReq.SetAcceptanceRadius(AcceptanceRadius);
-		MoveReq.SetAllowPartialPath(true);
-		MoveReq.SetUsePathfinding(true);
-		MoveReq.SetProjectGoalLocation(true);
-		MoveReq.SetRequireNavigableEndLocation(true);
-		MoveReq.SetNavigationFilter(AIController->GetDefaultNavigationFilterClass());
-		MoveReq.SetCanStrafe(false);
-
-		// request a move to the AI Controller
-		FNavPathSharedPtr FollowedPath;
-		const FPathFollowingRequestResult ResultData = AIController->MoveTo(MoveReq, &FollowedPath);
-		
-		// check the move result
-		switch (ResultData.Code)
-		{
-			// failed. Return false
-			case EPathFollowingRequestResult::Failed:
+	// cache the movement and interaction parameters
+	CurrentMovementGoal = Location;
+	bInteractOnArrival = bInteract;
+	InteractIgnoreList = IgnoreList;
 
-				return false;
-				break;
+	// stop movement and animation
+	StopMoving();
 
-			// already at goal. Return true and call the move completed delegate
-			case EPathFollowingRequestResult::AlreadyAtGoal:
+	// choose the EnvQuery to use
+	UEnvQuery* MoveQuery = bInteractOnArrival ? InteractionQuery : NoInteractionQuery;
 
-				OnMoveCompleted.Broadcast(this);
-				return true;
-				break;
+	// choose the run mode to use. The main interacting unit gets the closest result, all others choose randomly from top 25%
+	TEnumAsByte<EEnvQueryRunMode::Type> RunMode = bInteractOnArrival ? EEnvQueryRunMode::SingleResult : EEnvQueryRunMode::RandomBest25Pct;
 
-			// move successfully scheduled. Return true
-			case EPathFollowingRequestResult::RequestSuccessful:
+	// run an EQS to resolve the movement destination using the NavMesh
+	EnvQueryInstance = UEnvQueryManager::RunEQSQuery(this, MoveQuery, this,  RunMode, UEnvQueryInstanceBlueprintWrapper::StaticClass());
 
-				return true;
-				break;
-		}
+	if (IsValid(EnvQueryInstance))
+	{
+		EnvQueryInstance->GetOnQueryFinishedEvent().AddDynamic(this, &AStrategyUnit::OnEQSFinished);
 	}
+}
 
-	// the move could not be completed
-	return false;
+FVector AStrategyUnit::GetMovementGoal() const
+{
+	return CurrentMovementGoal;
+}
+
+void AStrategyUnit::OnEQSFinished(UEnvQueryInstanceBlueprintWrapper* QueryInstance, EEnvQueryStatus::Type QueryStatus)
+{
+	// was the EnvQuery successful?
+	if (QueryInstance)
+	{
+		// get the query result locations
+		TArray<FVector> ResultLocations;
+
+		if(QueryInstance->GetQueryResultsAsLocations(ResultLocations))
+		{
+			// grab the top result
+			CurrentMovementGoal = ResultLocations[0];
+
+			// ensure we have a valid AI Controller
+			if (AIController)
+			{
+				// set up the AI Move Request
+				FAIMoveRequest MoveReq;
+
+				MoveReq.SetGoalLocation(CurrentMovementGoal);
+				MoveReq.SetAcceptanceRadius(MovementAcceptanceRadius);
+				MoveReq.SetAllowPartialPath(true);
+				MoveReq.SetUsePathfinding(true);
+				MoveReq.SetProjectGoalLocation(true);
+				MoveReq.SetRequireNavigableEndLocation(true);
+				MoveReq.SetNavigationFilter(AIController->GetDefaultNavigationFilterClass());
+				MoveReq.SetCanStrafe(false);
+
+				// request a move to the AI Controller
+				FNavPathSharedPtr FollowedPath;
+				const FPathFollowingRequestResult ResultData = AIController->MoveTo(MoveReq, &FollowedPath);
+		
+				// check if we're already at the goal
+				if(ResultData.Code == EPathFollowingRequestResult::AlreadyAtGoal)
+				{
+					// finish movement immediately
+					HandleMoveFinished();
+				}
+			}
+		}
+	}
 }
 
 void AStrategyUnit::OnMoveFinished(FAIRequestID RequestID, const FPathFollowingResult& Result)
 {
-	// call the delegate
+	HandleMoveFinished();
+}
+
+void AStrategyUnit::HandleMoveFinished()
+{
+	// broadcast the move completed delegate
 	OnMoveCompleted.Broadcast(this);
+
+	if (bInteractOnArrival)
+	{
+		// do an overlap test to find nearby interactive objects
+		TArray<FOverlapResult> OutOverlaps;
+
+		FCollisionShape CollisionSphere;
+		CollisionSphere.SetSphere(InteractionRadius);
+
+		FCollisionObjectQueryParams ObjectParams;
+		ObjectParams.AddObjectTypesToQuery(ECC_WorldDynamic);
+
+		FCollisionQueryParams QueryParams;
+
+		// add the selected units to the ignored list
+		QueryParams.AddIgnoredActor(this);
+
+		for (const AActor* Current : InteractIgnoreList)
+		{
+			QueryParams.AddIgnoredActor(Current);
+		}
+
+		if (GetWorld()->OverlapMultiByObjectType(OutOverlaps, GetActorLocation(), FQuat::Identity, ObjectParams, CollisionSphere, QueryParams))
+		{
+			// find the first unit we've overlapped, and interact with it
+			for (const FOverlapResult& CurrentOverlap : OutOverlaps)
+			{
+				if (AStrategyUnit* CurrentUnit = Cast<AStrategyUnit>(CurrentOverlap.GetActor()))
+				{
+					CurrentUnit->Interact(this);
+					return;
+				}
+			}
+		}
+	}
 }
diff --git a/Templates/TP_TopDown/Source/TP_TopDown/Variant_Strategy/StrategyUnit.h b/Templates/TP_TopDown/Source/TP_TopDown/Variant_Strategy/StrategyUnit.h
--- a/Templates/TP_TopDown/Source/TP_TopDown/Variant_Strategy/StrategyUnit.h
+++ b/Templates/TP_TopDown/Source/TP_TopDown/Variant_Strategy/StrategyUnit.h
@@ -5,9 +5,12 @@
 #include "CoreMinimal.h"
 #include "GameFramework/Character.h"
 #include "AIController.h"
+#include "EnvironmentQuery/EnvQueryTypes.h"
 #include "StrategyUnit.generated.h"
 
 class USphereComponent;
+class UEnvQuery;
+class UEnvQueryInstanceBlueprintWrapper;
 
 /** Delegate to report that this unit has finished moving */
 DECLARE_DYNAMIC_MULTICAST_DELEGATE_OneParam(FOnUnitMoveCompletedDelegate, AStrategyUnit*, Unit);
@@ -55,14 +58,24 @@ public:
 	/** Notifies this unit that it's been interacted with by another actor */
 	void Interact(AStrategyUnit* Interactor);
 
-	/** Attempts to move this unit to its */
-	bool MoveToLocation(const FVector& Location, float AcceptanceRadius);
+	/** Attempts to move this unit to the passed location, and optionally signals it to interact on arrival */
+	void MoveToLocation(const FVector& Location, bool bInteract, const TArray<AStrategyUnit*> IgnoreList);
+
+	/** Returns the last cached movement goal location */
+	FVector GetMovementGoal() const;
 
 protected:
 
-	/** called by the AI controller when this unit has finished moving */
+	/** Called by EQS when the movement destination query has finished */
+	UFUNCTION()
+	void OnEQSFinished(UEnvQueryInstanceBlueprintWrapper* QueryInstance, EEnvQueryStatus::Type QueryStatus);
+
+	/** Called by the AI controller when this unit has finished moving */
 	void OnMoveFinished(FAIRequestID RequestID, const FPathFollowingResult& Result);
 
+	/** Wraps up movement logic */
+	void HandleMoveFinished();
+
 protected:
 
 	/** Blueprint handler for strategy game selection */
@@ -73,10 +86,44 @@ protected:
 	UFUNCTION(BlueprintImplementableEvent, Category="NPC", meta = (DisplayName="Unit Deselected"))
 	void BP_UnitDeselected();
 
+	/** Blueprint handler to stop the unit's interaction animation */
+	UFUNCTION(BlueprintImplementableEvent, BlueprintCallable, Category="NPC", meta = (DisplayName="Stop Animation"))
+	void BP_StopAnimation();
+
 	/** Blueprint handler for strategy game interactions */
 	UFUNCTION(BlueprintImplementableEvent, Category="NPC", meta = (DisplayName="Interaction Behavior"))
 	void BP_InteractionBehavior(AStrategyUnit* Interactor);
 
+protected:
+
+	/** EnvQuery to use when this unit interacts after movement */
+	UPROPERTY(EditAnywhere, Category="NPC")
+	TObjectPtr<UEnvQuery> InteractionQuery;
+
+	/** EnvQuery to use when this unit does not interact after movement */
+	UPROPERTY(EditAnywhere, Category="NPC")
+	TObjectPtr<UEnvQuery> NoInteractionQuery;
+
+	/** How close we should get to the movement goal to consider ourselves as having reached it */
+	UPROPERTY(EditAnywhere, Category="NPC", meta = (ClampMin = 0, ClampMax = 10000, Units = "cm"))
+	float MovementAcceptanceRadius = 100.0f;
+
+	/** Max distance to look for nearby units when doing an interaction check */
+	UPROPERTY(EditAnywhere, Category="Input", meta = (ClampMin = 0, ClampMax = 10000, Units = "cm"))
+	float InteractionRadius = 250.0f;
+
+	/** EQS instance running the movement query for this unit */
+	TObjectPtr<UEnvQueryInstanceBlueprintWrapper> EnvQueryInstance;
+
+	/** Cached movement goal for this unit */
+	FVector CurrentMovementGoal;
+
+	/** If true, this unit will attempt to interact with a nearby unit upon finishing movement */
+	bool bInteractOnArrival = false;
+
+	/** List of actors to ignore when searching for units to interact with */
+	TArray<AStrategyUnit*> InteractIgnoreList;
+
 public:
 
 	FOnUnitMoveCompletedDelegate OnMoveCompleted;
diff --git a/Templates/TP_TopDown/Source/TP_TopDown/Variant_Strategy/UI/StrategyTouchControls.cpp b/Templates/TP_TopDown/Source/TP_TopDown/Variant_Strategy/UI/StrategyTouchControls.cpp
new file mode 100644
--- /dev/null
+++ b/Templates/TP_TopDown/Source/TP_TopDown/Variant_Strategy/UI/StrategyTouchControls.cpp
@@ -0,0 +1,36 @@
+// Copyright Epic Games, Inc. All Rights Reserved.
+
+
+#include "Variant_Strategy/UI/StrategyTouchControls.h"
+#include "Variant_Strategy/StrategyPlayerController.h"
+
+void UStrategyTouchControls::SetPlayerController(AStrategyPlayerController* PC)
+{
+	PlayerController = PC;
+}
+
+void UStrategyTouchControls::ResetZoom()
+{
+	if (PlayerController)
+	{
+		PlayerController->DoCameraResetZoomCommand();
+
+		BP_SetZoomPercentage(PlayerController->GetDefaultZoomPercentage());
+	}
+}
+
+void UStrategyTouchControls::ToggleSelectAllUnits()
+{
+	if (PlayerController)
+	{
+		PlayerController->DoToggleSelectAllUnitsCommand();
+	}
+}
+
+void UStrategyTouchControls::SetZoomPercentage(float Percentage)
+{
+	if (PlayerController)
+	{
+		PlayerController->DoCameraSetZoomPercentageCommand(Percentage);
+	}
+}
\ No newline at end of file
diff --git a/Templates/TP_TopDown/Source/TP_TopDown/Variant_Strategy/UI/StrategyTouchControls.h b/Templates/TP_TopDown/Source/TP_TopDown/Variant_Strategy/UI/StrategyTouchControls.h
new file mode 100644
--- /dev/null
+++ b/Templates/TP_TopDown/Source/TP_TopDown/Variant_Strategy/UI/StrategyTouchControls.h
@@ -0,0 +1,47 @@
+// Copyright Epic Games, Inc. All Rights Reserved.
+
+#pragma once
+
+#include "CoreMinimal.h"
+#include "Blueprint/UserWidget.h"
+#include "StrategyTouchControls.generated.h"
+
+class AStrategyPlayerController;
+
+/**
+ *  Base class for additional touchscreen controls for a strategy game.
+ *  Exposes some game commands to UI
+ */
+UCLASS(abstract)
+class TP_TOPDOWN_API UStrategyTouchControls : public UUserWidget
+{
+	GENERATED_BODY()
+	
+protected:
+
+	/** Pointer to the owning Strategy PC */
+	TObjectPtr<AStrategyPlayerController> PlayerController;
+
+public:
+
+	/** Sets the owning Strategy PC pointer */
+	void SetPlayerController(AStrategyPlayerController* PC);
+
+	/** Syncs the camera zoom percentage with the UI. Called by the owning PC */
+	UFUNCTION(BlueprintImplementableEvent, Category="UI", meta=(DisplayName="Set Zoom Percentage"))
+	void BP_SetZoomPercentage(float Percentage);
+
+protected:
+
+	/** Resets the camera zoom level */
+	UFUNCTION(BlueprintCallable, Category="UI")
+	void ResetZoom();
+
+	/** Toggles between select all units and deselect all units. */
+	UFUNCTION(BlueprintCallable, Category="UI")
+	void ToggleSelectAllUnits();
+
+	/** Sets the camera zoom percentage level */
+	UFUNCTION(BlueprintCallable, Category="UI")
+	void SetZoomPercentage(float Percentage);
+};
diff --git a/Templates/TP_TopDown/TP_TopDown.uproject b/Templates/TP_TopDown/TP_TopDown.uproject
--- a/Templates/TP_TopDown/TP_TopDown.uproject
+++ b/Templates/TP_TopDown/TP_TopDown.uproject
@@ -7,7 +7,11 @@
 		{
 			"Name": "TP_TopDown",
 			"Type": "Runtime",
-			"LoadingPhase": "Default"
+			"LoadingPhase": "Default",
+			"AdditionalDependencies": [
+				"UMG",
+				"AIModule"
+			]
 		}
 	],
 	"Plugins": [
@@ -26,5 +30,5 @@
 			"Name": "GameplayStateTree",
 			"Enabled": true
 		}
-	]	
+	]
 }
\ No newline at end of file
